Composite: Moved plugin from NotForLicensees/ back to Experimental/Compositing/.

#virtualized

diff --git a/Engine/Plugins/Experimental/Compositing/Composite/Composite.uplugin b/Engine/Plugins/Experimental/Compositing/Composite/Composite.uplugin
new file mode 100644
--- /dev/null
+++ b/Engine/Plugins/Experimental/Compositing/Composite/Composite.uplugin
@@ -0,0 +1,40 @@
+{
+	"FileVersion": 1,
+	"Version": 1,
+	"VersionName": "1.0",
+	"FriendlyName": "Composite",
+	"Description": "Plugin designed to facilitate real-time compositing pipelines, extending CompositeCore.",
+	"Category": "Compositing",
+	"CreatedBy": "Epic Games, Inc.",
+	"CreatedByURL": "https://epicgames.com",
+	"DocsURL": "",
+	"MarketplaceURL": "",
+	"SupportURL": "",
+	"EnabledByDefault": false,
+	"CanContainContent": true,
+	"IsBetaVersion": false,
+	"IsExperimentalVersion": true,
+	"Installed": false,
+	"Modules": [
+		{
+			"Name": "Composite",
+			"Type": "Runtime",
+			"LoadingPhase": "PostConfigInit"
+		},
+		{
+			"Name": "CompositeEditor",
+			"Type": "Editor",
+			"LoadingPhase": "Default"
+		}
+	],
+	"Plugins": [
+		{
+			"Name": "CompositeCore",
+			"Enabled": true
+		},
+		{
+			"Name": "OpenColorIO",
+			"Enabled": true
+		}
+	]
+}
\ No newline at end of file
diff --git a/Engine/Plugins/Experimental/Compositing/Composite/Config/BaseComposite.ini b/Engine/Plugins/Experimental/Compositing/Composite/Config/BaseComposite.ini
new file mode 100644
--- /dev/null
+++ b/Engine/Plugins/Experimental/Compositing/Composite/Config/BaseComposite.ini
@@ -0,0 +1,33 @@
+[CoreRedirects]
+
++ClassRedirects=(OldName="/Script/HoldoutComposite.HoldoutComposite", NewName="/Script/Composite.CompositeActor")
++ClassRedirects=(OldName="/Script/Composite.Composite", NewName="/Script/Composite.CompositeActor")
++ClassRedirects=(OldName="/Script/HoldoutComposite.HoldoutCompositePassBase", NewName="/Script/Composite.CompositePassBase")
++ClassRedirects=(OldName="/Script/HoldoutComposite.HoldoutCompositePassMaterial", NewName="/Script/Composite.CompositePassMaterial")
++ClassRedirects=(OldName="/Script/HoldoutComposite.HoldoutCompositePassFXAA", NewName="/Script/Composite.CompositePassFXAA")
++ClassRedirects=(OldName="/Script/HoldoutComposite.HoldoutCompositeViewProjectionComponent", NewName="/Script/Composite.CompositeViewProjectionComponent")
+
++EnumRedirects=(OldName="/Script/HoldoutComposite.HoldoutCompositePassInput",NewName="/Script/Composite.CompositePassInput")
+
++StructRedirects=(OldName="/Script/HoldoutComposite.HoldoutCompositeProcessingPasses", NewName="/Script/Composite.CompositeProcessingPasses")
+
++PackageRedirects=(OldName="/HoldoutComposite/MaterialFunctions/MF_HoldoutComposite_ProjectedUVs", NewName="/Composite/MaterialFunctions/MF_Composite_ProjectedUVs")
++PackageRedirects=(OldName="/HoldoutComposite/Materials/MPC_HoldoutComposite", NewName="/Composite/Materials/MPC_Composite")
++PackageRedirects=(OldName="/HoldoutComposite/Materials/M_HoldoutComposite_KeyedProjection", NewName="/Composite/Materials/M_Composite_KeyedProjection")
+
+; 5.7
+
++ClassRedirects=(OldName="/Script/Composite.CompositePassLayerBase", NewName="/Script/Composite.CompositeLayerBase")
++ClassRedirects=(OldName="/Script/Composite.CompositePassLayerMainRender", NewName="/Script/Composite.CompositeLayerMainRender")
++ClassRedirects=(OldName="/Script/Composite.CompositePassLayerPlate", NewName="/Script/Composite.CompositeLayerPlate")
++ClassRedirects=(OldName="/Script/Composite.CompositePassLayerShadowReflection", NewName="/Script/Composite.CompositeLayerShadowReflection")
++ClassRedirects=(OldName="/Script/Composite.CompositeCenteredScalePass", NewName="/Script/Composite.CompositePassCenteredScale")
++ClassRedirects=(OldName="/Script/Composite.CompositeFXAAPass", NewName="/Script/Composite.CompositePassFXAA")
+
++PropertyRedirects=(OldName="/Script/Composite.CompositeActor.CompositeMeshActors", NewName="/Script/Composite.CompositeLayerPlate.CompositeMeshes")
++PropertyRedirects=(OldName="/Script/Composite.CompositeLayerPlate.Passes", NewName="/Script/Composite.CompositeLayerPlate.Postprocessing")
++PropertyRedirects=(OldName="/Script/Composite.CompositeLayerPlate.Subpasses", NewName="/Script/Composite.CompositeLayerPlate.Postprocessing")
++PropertyRedirects=(OldName="/Script/Composite.CompositeLayerPlate.MediaTexture", NewName="/Script/Composite.CompositeLayerPlate.Texture")
+
+
+
diff --git a/Engine/Plugins/Experimental/Compositing/Composite/Shaders/Private/CompositeCenteredScale.usf b/Engine/Plugins/Experimental/Compositing/Composite/Shaders/Private/CompositeCenteredScale.usf
new file mode 100644
--- /dev/null
+++ b/Engine/Plugins/Experimental/Compositing/Composite/Shaders/Private/CompositeCenteredScale.usf
@@ -0,0 +1,19 @@
+// Copyright Epic Games, Inc. All Rights Reserved.
+
+#include "/Engine/Private/Common.ush"
+
+Texture2D InputTexture;
+SamplerState InputSampler;
+float2 ScaleUV;
+static const float2 CenterUV = 0.5f;
+
+void MainPS(
+    noperspective float4 InUVAndScreenPos : TEXCOORD0,
+    out float4 OutColor : SV_Target0
+    )
+{
+    float2 UV = InUVAndScreenPos.xy;
+    UV = ((UV - CenterUV) * ScaleUV) + CenterUV;
+
+    OutColor = InputTexture.Sample(InputSampler, UV);
+}
diff --git a/Engine/Plugins/Experimental/Compositing/Composite/Source/Composite/Composite.Build.cs b/Engine/Plugins/Experimental/Compositing/Composite/Source/Composite/Composite.Build.cs
new file mode 100644
--- /dev/null
+++ b/Engine/Plugins/Experimental/Compositing/Composite/Source/Composite/Composite.Build.cs
@@ -0,0 +1,51 @@
+// Copyright Epic Games, Inc. All Rights Reserved.
+
+using UnrealBuildTool;
+
+public class Composite : ModuleRules
+{
+	public Composite(ReadOnlyTargetRules Target) : base(Target)
+	{
+		// The Composite plugin is distributed with engine hot fixes and thus isn't tied to binary
+		// compatibility between hotfixes by only using Public/ interface of the renderer, but also Internal/ ones.
+		bTreatAsEngineModule = true;
+		PCHUsage = PCHUsageMode.UseExplicitOrSharedPCHs;
+
+		PublicDependencyModuleNames.AddRange(
+			new string[]
+			{
+				"Core",
+				"CoreUObject",
+				"Engine",
+				"CompositeCore",
+				"OpenColorIO",
+				"MediaAssets",
+			}
+		);
+
+		PrivateDependencyModuleNames.AddRange(
+			new string[]
+			{
+				"Projects",
+				"RenderCore",
+				"Renderer",
+				"RHI",
+			}
+		);
+
+		if (Target.bBuildEditor)
+		{
+			PrivateDependencyModuleNames.AddRange(
+				new string[]
+				{
+					"UnrealEd",
+					"Slate",
+				}
+			);
+		}
+
+		// bUseUnity = false;
+		// OptimizeCode = CodeOptimization.Never;
+		// PCHUsage = ModuleRules.PCHUsageMode.NoPCHs;
+	}
+}
diff --git a/Engine/Plugins/Experimental/Compositing/Composite/Source/Composite/Private/Components/CompositeMeshComponent.cpp b/Engine/Plugins/Experimental/Compositing/Composite/Source/Composite/Private/Components/CompositeMeshComponent.cpp
new file mode 100644
--- /dev/null
+++ b/Engine/Plugins/Experimental/Compositing/Composite/Source/Composite/Private/Components/CompositeMeshComponent.cpp
@@ -0,0 +1,70 @@
+// Copyright Epic Games, Inc. All Rights Reserved.
+
+#include "Components/CompositeMeshComponent.h"
+
+#include "Engine/StaticMesh.h"
+#include "Materials/Material.h"
+#include "UObject/ConstructorHelpers.h"
+
+#define LOCTEXT_NAMESPACE "Composite"
+
+UCompositeMeshComponent::UCompositeMeshComponent(const FObjectInitializer& ObjectInitializer)
+	: Super(ObjectInitializer)
+{
+	struct FConstructorStatics
+	{
+		ConstructorHelpers::FObjectFinder<UStaticMesh> BackgroundCube;
+		ConstructorHelpers::FObjectFinder<UMaterial> DefaultMaterial;
+		FConstructorStatics()
+			: BackgroundCube(TEXT("/Engine/EngineMeshes/BackgroundCube.BackgroundCube"))
+			, DefaultMaterial(TEXT("/Composite/Materials/M_Composite_KeyedProjection.M_Composite_KeyedProjection"))
+		{
+		}
+	};
+
+	static FConstructorStatics ConstructorStatics;
+	if (ensure(ConstructorStatics.BackgroundCube.Object != nullptr))
+	{
+		SetStaticMesh(ConstructorStatics.BackgroundCube.Object);
+		SetMaterial(0, ConstructorStatics.DefaultMaterial.Object);
+	}
+}
+
+void UCompositeMeshComponent::PostInitProperties()
+{
+	Super::PostInitProperties();
+}
+
+void UCompositeMeshComponent::OnRegister()
+{
+	Super::OnRegister();
+}
+
+void UCompositeMeshComponent::PostDuplicate(bool bDuplicateForPIE)
+{
+	Super::PostDuplicate(bDuplicateForPIE);
+}
+
+void UCompositeMeshComponent::PostEditImport()
+{
+	Super::PostEditImport();
+}
+
+void UCompositeMeshComponent::PostLoad()
+{
+	Super::PostLoad();
+}
+
+void UCompositeMeshComponent::OnUnregister()
+{
+	Super::OnUnregister();
+}
+
+#if WITH_EDITOR
+void UCompositeMeshComponent::PostEditChangeProperty(FPropertyChangedEvent& PropertyChangedEvent)
+{
+	Super::PostEditChangeProperty(PropertyChangedEvent);
+}
+#endif
+
+#undef LOCTEXT_NAMESPACE
diff --git a/Engine/Plugins/Experimental/Compositing/Composite/Source/Composite/Private/Components/CompositeShadowReflectionCatcherComponent.cpp b/Engine/Plugins/Experimental/Compositing/Composite/Source/Composite/Private/Components/CompositeShadowReflectionCatcherComponent.cpp
new file mode 100644
--- /dev/null
+++ b/Engine/Plugins/Experimental/Compositing/Composite/Source/Composite/Private/Components/CompositeShadowReflectionCatcherComponent.cpp
@@ -0,0 +1,73 @@
+// Copyright Epic Games, Inc. All Rights Reserved.
+
+#include "Components/CompositeShadowReflectionCatcherComponent.h"
+
+#include "Engine/TextureRenderTarget2D.h"
+#include "UObject/Package.h"
+
+#define LOCTEXT_NAMESPACE "Composite"
+
+UCompositeShadowReflectionCatcherComponent::UCompositeShadowReflectionCatcherComponent(const FObjectInitializer& ObjectInitializer)
+	: Super(ObjectInitializer)
+{
+	PrimitiveRenderMode = ESceneCapturePrimitiveRenderMode::PRM_RenderScenePrimitives;
+	CaptureSource = ESceneCaptureSource::SCS_FinalColorHDR;
+
+	bUseRayTracingIfEnabled = true;
+	bMainViewFamily = true;
+	bMainViewResolution = true;
+	bMainViewCamera = true;
+
+	ShowFlags.SetAtmosphere(false);
+	ShowFlags.SetAllowPrimitiveAlphaHoldout(false);
+	ShowFlags.SetGrain(false);
+	ShowFlags.SetBloom(false);
+	ShowFlags.SetEyeAdaptation(false);
+	ShowFlags.SetLocalExposure(false);
+	ShowFlags.SetPostProcessMaterial(false);
+	ShowFlags.SetScreenSpaceReflections(false);
+	//Note: This can cause temporal ghosting artifacts on moving shadows, but remains preferable over the otherwise default shadow noise.
+	ShowFlags.SetTemporalAA(true);
+
+	PostProcessSettings.bOverride_ReflectionMethod = true;
+	PostProcessSettings.ReflectionMethod = EReflectionMethod::Lumen;
+
+	PostProcessSettings.bOverride_DynamicGlobalIlluminationMethod = true;
+	PostProcessSettings.DynamicGlobalIlluminationMethod = EDynamicGlobalIlluminationMethod::Lumen;
+
+	PostProcessSettings.bOverride_LumenFinalGatherScreenTraces = true;
+	PostProcessSettings.LumenFinalGatherScreenTraces = false;
+
+	PostProcessSettings.bOverride_LumenReflectionsScreenTraces = true;
+	PostProcessSettings.LumenReflectionsScreenTraces = false;
+}
+
+UTextureRenderTarget2D* UCompositeShadowReflectionCatcherComponent::FindOrCreateRenderTarget(int32 InSizeX, int32 InSizeY)
+{
+	// TODO: Temporary solution, evaluate alternatives.
+	if (!IsValid(TextureTarget))
+	{
+		const FName RTName = MakeUniqueObjectName(GetTransientPackage(), UTextureRenderTarget2D::StaticClass(), TEXT("ShadowReflectionCatcherRT"));
+		TextureTarget = NewObject<UTextureRenderTarget2D>(GetTransientPackage(), RTName, RF_Transient);
+		TextureTarget->SizeX = InSizeX;
+		TextureTarget->SizeY = InSizeY;
+		TextureTarget->RenderTargetFormat = ETextureRenderTargetFormat::RTF_RGBA16f;
+		TextureTarget->bForceLinearGamma = true;
+		//TextureTarget->bSupportsUAV = true;
+
+		TextureTarget->UpdateResource();
+	}
+
+	TextureTarget->ResizeTarget(InSizeX, InSizeY);
+
+	return TextureTarget;
+}
+
+#if WITH_EDITOR
+void UCompositeShadowReflectionCatcherComponent::PostEditChangeProperty(FPropertyChangedEvent& PropertyChangedEvent)
+{
+	Super::PostEditChangeProperty(PropertyChangedEvent);
+}
+#endif
+
+#undef LOCTEXT_NAMESPACE
diff --git a/Engine/Plugins/Experimental/Compositing/Composite/Source/Composite/Private/Components/CompositeViewProjectionComponent.cpp b/Engine/Plugins/Experimental/Compositing/Composite/Source/Composite/Private/Components/CompositeViewProjectionComponent.cpp
new file mode 100644
--- /dev/null
+++ b/Engine/Plugins/Experimental/Compositing/Composite/Source/Composite/Private/Components/CompositeViewProjectionComponent.cpp
@@ -0,0 +1,210 @@
+// Copyright Epic Games, Inc. All Rights Reserved.
+
+#include "Components/CompositeViewProjectionComponent.h"
+#include "CompositeActor.h"
+
+#include "Camera/CameraComponent.h"
+#include "Engine/Engine.h"
+#include "Engine/World.h"
+#include "CompositeCoreSubsystem.h"
+#include "Kismet/GameplayStatics.h"
+#include "Materials/MaterialParameterCollection.h"
+#include "Materials/MaterialParameterCollectionInstance.h"
+#include "UObject/ConstructorHelpers.h"
+
+#define LOCTEXT_NAMESPACE "Composite"
+
+UCompositeViewProjectionComponent::UCompositeViewProjectionComponent(const FObjectInitializer& ObjectInitializer)
+	: Super(ObjectInitializer)
+{
+	static ConstructorHelpers::FObjectFinder<UMaterialParameterCollection> DefaultMPC = TEXT("/Composite/Materials/MPC_Composite.MPC_Composite");
+
+	MaterialParameterCollection = DefaultMPC.Object;
+	MatrixParameterName = TEXT("CameraViewProjectionMatrix");
+
+	PrimaryComponentTick.bStartWithTickEnabled = true;
+	PrimaryComponentTick.bCanEverTick = true;
+	PrimaryComponentTick.bTickEvenWhenPaused = true;
+	bTickInEditor = true;
+	bAutoActivate = true;
+}
+
+void UCompositeViewProjectionComponent::ForceUpdate()
+{
+	LastViewProjectionMatrix = FMatrix::Identity;
+}
+
+FComponentReference& UCompositeViewProjectionComponent::GetTargetComponent()
+{
+	return TargetCameraComponent;
+}
+
+void UCompositeViewProjectionComponent::SetTargetComponent(const FComponentReference& InComponentReference)
+{
+	TargetCameraComponent = InComponentReference;
+	LastCameraComponent = Cast<UCameraComponent>(TargetCameraComponent.GetComponent(nullptr));
+
+	ForceUpdate();
+}
+
+void UCompositeViewProjectionComponent::InitDefaultCamera()
+{
+	if (ACompositeActor* CompositeActor = GetTypedOuter<ACompositeActor>())
+	{
+		// Special case when the component is owned by the composite actor: we use its camera component reference directly.
+		SetTargetComponent(CompositeActor->GetCamera());
+
+		return;
+	}
+
+	UCameraComponent* TargetCamera = Cast<UCameraComponent>(TargetCameraComponent.GetComponent(GetOwner()));
+	
+	if (!IsValid(TargetCamera))
+	{
+		// Find all CineCameraComponents on the same actor as this component and set the first one to be the target
+		TInlineComponentArray<UCameraComponent*> CameraComponents;
+		if (AActor* Owner = GetOwner())
+		{
+			Owner->GetComponents(CameraComponents);
+		}
+
+		if (CameraComponents.Num() > 0)
+		{
+			UCameraComponent* TargetComponent = CameraComponents[0];
+			TargetCameraComponent.ComponentProperty = TargetComponent->GetFName();
+			LastCameraComponent = TargetComponent;
+		}
+	}
+}
+
+void UCompositeViewProjectionComponent::PostInitProperties()
+{
+	Super::PostInitProperties();
+}
+
+void UCompositeViewProjectionComponent::OnRegister()
+{
+	Super::OnRegister();
+
+	InitDefaultCamera();
+
+	/*
+	* NOTE: We must (unfortunately) rely on this after-tick event to update the view matrix projection,
+	* since viewport camera transform updates after ticking all world objects.
+	*
+	* The CompositePlane plugin updated this information on tick, but a lag becomes visible when
+	* piloting the source projection camera for example.
+	*/
+	OnWorldPreSendAllEndOfFrameUpdatesHandle = FWorldDelegates::OnWorldPreSendAllEndOfFrameUpdates.AddUObject(this, &UCompositeViewProjectionComponent::UpdateProjection);
+}
+
+void UCompositeViewProjectionComponent::PostDuplicate(bool bDuplicateForPIE)
+{
+	Super::PostDuplicate(bDuplicateForPIE);
+
+	InitDefaultCamera();
+}
+
+void UCompositeViewProjectionComponent::PostEditImport()
+{
+	Super::PostEditImport();
+
+	InitDefaultCamera();
+}
+
+void UCompositeViewProjectionComponent::PostLoad()
+{
+	Super::PostLoad();
+}
+
+void UCompositeViewProjectionComponent::OnUnregister()
+{
+	FWorldDelegates::OnWorldPreSendAllEndOfFrameUpdates.Remove(OnWorldPreSendAllEndOfFrameUpdatesHandle);
+
+	Super::OnUnregister();
+}
+
+void UCompositeViewProjectionComponent::UpdateProjection(UWorld* InWorld) const
+{
+	if (!IsValid(InWorld))
+	{
+		return;
+	}
+
+	if (!bIsEnabled)
+	{
+		return;
+	}
+
+	UMaterialParameterCollection* MPC = MaterialParameterCollection.Get();
+
+	if (IsValid(LastCameraComponent) && IsValid(MPC))
+	{
+		FMinimalViewInfo DesiredView;
+		FMatrix ViewMatrix, ProjectionMatrix, ViewProjectionMatrix;
+
+		LastCameraComponent->GetCameraView(FApp::GetDeltaTime(), DesiredView);
+
+		UGameplayStatics::GetViewProjectionMatrix(DesiredView, ViewMatrix, ProjectionMatrix, ViewProjectionMatrix);
+
+		if (!LastViewProjectionMatrix.Equals(ViewProjectionMatrix))
+		{
+			int32 ParameterProjection = INDEX_NONE;
+			for (int32 ParameterIndex = 0; ParameterIndex < MPC->VectorParameters.Num(); ParameterIndex++)
+			{
+				FName CurrentParameterName = MPC->VectorParameters[ParameterIndex].ParameterName;
+				if (!CurrentParameterName.IsNone())
+				{
+					if (CurrentParameterName == MatrixParameterName)
+					{
+						ParameterProjection = ParameterIndex;
+					}
+				}
+			}
+
+			UMaterialParameterCollectionInstance* MaterialParameterCollectionInstance = GetWorld()->GetParameterCollectionInstance(MaterialParameterCollection);
+			// Ensure there's space for 4 output vectors for the projection matrix
+			if (ParameterProjection != INDEX_NONE && ParameterProjection + 4 <= MaterialParameterCollection->VectorParameters.Num())
+			{
+				FMatrix44f ViewProjectionMatrixf = FMatrix44f(ViewProjectionMatrix);
+
+				// Store the vectors to the collection instance
+				const FLinearColor* MatrixVectors = (const FLinearColor*)&ViewProjectionMatrixf;
+				for (int32 ElementIndex = 0; ElementIndex < 4; ElementIndex++)
+				{
+					MaterialParameterCollectionInstance->SetVectorParameterValue(MaterialParameterCollection->VectorParameters[ParameterProjection + ElementIndex].ParameterName, MatrixVectors[ElementIndex]);
+				}
+
+				MaterialParameterCollectionInstance->UpdateRenderState(false);
+			}
+
+			LastViewProjectionMatrix = ViewProjectionMatrix;
+		}
+	}
+}
+
+void UCompositeViewProjectionComponent::TickComponent(float DeltaTime, ELevelTick TickType, FActorComponentTickFunction* ThisTickFunction)
+{
+	Super::TickComponent(DeltaTime, TickType, ThisTickFunction);
+
+	MarkForNeededEndOfFrameUpdate();
+}
+
+#if WITH_EDITOR
+void UCompositeViewProjectionComponent::PostEditChangeProperty(FPropertyChangedEvent& PropertyChangedEvent)
+{
+	const FName PropertyName = (PropertyChangedEvent.Property != NULL) ? PropertyChangedEvent.Property->GetFName() : NAME_None;
+
+	if (PropertyName == GET_MEMBER_NAME_CHECKED(ThisClass, TargetCameraComponent))
+	{
+		LastCameraComponent = Cast<UCameraComponent>(TargetCameraComponent.GetComponent(GetOwner()));
+	}
+
+	// Force MPC updates on any property change.
+	ForceUpdate();
+
+	Super::PostEditChangeProperty(PropertyChangedEvent);
+}
+#endif
+
+#undef LOCTEXT_NAMESPACE
diff --git a/Engine/Plugins/Experimental/Compositing/Composite/Source/Composite/Private/CompositeActor.cpp b/Engine/Plugins/Experimental/Compositing/Composite/Source/Composite/Private/CompositeActor.cpp
new file mode 100644
--- /dev/null
+++ b/Engine/Plugins/Experimental/Compositing/Composite/Source/Composite/Private/CompositeActor.cpp
@@ -0,0 +1,571 @@
+// Copyright Epic Games, Inc. All Rights Reserved.
+
+#include "CompositeActor.h"
+
+#include "CompositeCoreSubsystem.h"
+#include "CompositeModule.h"
+#include "CompositeCoreSettings.h"
+
+#include "Components/CompositeShadowReflectionCatcherComponent.h"
+#include "Components/CompositeViewProjectionComponent.h"
+#include "Layers/CompositeLayerShadowReflection.h"
+#include "Layers/CompositeLayerPlate.h"
+#include "Layers/CompositeLayerMainRender.h"
+
+#include "Algo/Rotate.h"
+#include "Components/PrimitiveComponent.h"
+#include "Camera/CameraActor.h"
+#include "Camera/CameraComponent.h"
+#include "Engine/Engine.h"
+#include "Engine/World.h"
+#include "Engine/Texture.h"
+#include "Kismet/GameplayStatics.h"
+#include "Materials/MaterialInstanceDynamic.h"
+#include "SceneView.h"
+#include "UObject/ObjectSaveContext.h"
+#include "UObject/UObjectIterator.h"
+
+#if WITH_EDITOR
+#include "LevelEditor.h"
+#include "Modules/ModuleManager.h"
+
+/** The editor object. */
+extern UNREALED_API class UEditorEngine* GEditor;
+#endif
+
+#include UE_INLINE_GENERATED_CPP_BY_NAME(CompositeActor)
+
+DECLARE_CYCLE_STAT(TEXT("Composite - Actor Tick"), STAT_Composite_Actor_Tick, STATGROUP_Composite);
+DECLARE_CYCLE_STAT(TEXT("Composite - Actor EndOfFrame Update"), STAT_Composite_Actor_EndOfFrameUpdate, STATGROUP_Composite);
+
+namespace UE
+{
+	namespace Composite
+	{
+		namespace Private
+		{
+			/** Return solo layer from list. */
+			TObjectPtr<UCompositeLayerBase> FindSoloLayer(TArrayView<const TObjectPtr<UCompositeLayerBase>> InLayers)
+			{
+				TObjectPtr<UCompositeLayerBase> SoloLayer;
+
+				// Iterate backwards so that lower passes in the UI are executed first
+				for (const TObjectPtr<UCompositeLayerBase>& Layer : InLayers)
+				{
+					if (!IsValid(Layer) || !Layer->bIsEnabled)
+					{
+						continue;
+					}
+
+					if (Layer->bIsSolo)
+					{
+						SoloLayer = Layer;
+					}
+				}
+
+				return SoloLayer;
+			}
+
+			/** Convenience function to register layers in the traversal context, while returning proxies for each layer input. */
+			void RegisterLayers(UCompositeLayerBase::FTraversalContext& Context, TArrayView<const TObjectPtr<UCompositeLayerBase>> InLayers, FSceneRenderingBulkObjectAllocator& InFrameAllocator, TArray<const FCompositeCorePassProxy*>& OutProxies)
+			{
+				OutProxies.Reserve(InLayers.Num());
+				Context.bIsSolo = (InLayers.Num() == 1);
+
+				// Iterate backwards so that lower passes in the UI are executed first
+				const int32 FirstPassIndex = InLayers.Num() - 1;
+				for (int32 LayerIndex = FirstPassIndex; LayerIndex >= 0; --LayerIndex)
+				{
+					const TObjectPtr<UCompositeLayerBase>& Layer = InLayers[LayerIndex];
+					if (!IsValid(Layer) || !Layer->bIsEnabled)
+					{
+						continue;
+					}
+
+					Context.bIsFirstPass = (LayerIndex == FirstPassIndex);
+
+					FCompositeCorePassProxy* OutProxy;
+					if (Layer->GetProxy(Context, InFrameAllocator, OutProxy))
+					{
+						OutProxies.Add(OutProxy);
+					}
+				}
+			}
+
+			/* Get post-processing render work (per frame) given the current composite actor properties. */
+			UE::CompositeCore::FPostRenderWork GetPostRenderWork(const ACompositeActor& InCompositeActor)
+			{
+				if (InCompositeActor.CompositeLayers.IsEmpty())
+				{
+					return {};
+				}
+
+				UE::CompositeCore::FPostRenderWork RenderWork;
+				TArray<TObjectPtr<UCompositeLayerBase>> Layers;
+
+				TObjectPtr<UCompositeLayerBase> SoloLayer = FindSoloLayer(InCompositeActor.CompositeLayers);
+				if (IsValid(SoloLayer))
+				{
+					Layers.Add(SoloLayer);
+				}
+				else
+				{
+					Layers = InCompositeActor.CompositeLayers;
+				}
+
+				{
+					UCompositeLayerBase::FTraversalContext Context;
+					RegisterLayers(Context, Layers, *RenderWork.FrameAllocator, RenderWork.FramePasses.Add(ISceneViewExtension::EPostProcessingPass::Tonemap));
+					
+					RenderWork.ExternalInputs = Context.ExternalTextures;
+					RenderWork.PreprocessingPasses = Context.PreprocessingPasses;
+				}
+
+				// TODO: Recycle work between SSR input & after-tonemap composite step.
+				if (InCompositeActor.bEnableScreenSpaceReflections)
+				{
+					UCompositeLayerBase::FTraversalContext TempContext;
+					RegisterLayers(TempContext, Layers, *RenderWork.FrameAllocator, RenderWork.FramePasses.Add(ISceneViewExtension::EPostProcessingPass::SSRInput));
+				}
+
+				return RenderWork;
+			}
+
+			void SetCaptureEveryFrame(TMap<TWeakObjectPtr<const UCompositeLayerBase>, TObjectPtr<UCompositeShadowReflectionCatcherComponent>>& InComponents, bool bState)
+			{
+				for (auto& Pair : InComponents)
+				{
+					TObjectPtr<UCompositeShadowReflectionCatcherComponent>& Component = Pair.Value;
+					if (IsValid(Component))
+					{
+						Component->bCaptureEveryFrame = bState;
+					}
+				}
+			}
+		}
+	}
+}
+
+ACompositeActor::ACompositeActor(const FObjectInitializer& ObjectInitializer)
+	: Super(ObjectInitializer)
+	, ViewProjectionComponent{}
+	, bIsEnabled{true}
+// TODO: Can we find the current viewport render resolution instead?
+// View rects are only accessible from the render thread, but maybe there is an alternative.
+// This isn't needed by custom render passes, only regular scene capture render targets...
+	, RenderResolution{1920, 1080}
+	, bEnableScreenSpaceReflections{true}
+	, bOverridesViewUserFlags{true}
+	, ViewUserFlags{1}
+	, bIsModifyingAProperty{false}
+{
+	PrimaryActorTick.bCanEverTick = true;
+	PrimaryActorTick.bTickEvenWhenPaused = true;
+	PrimaryActorTick.TickGroup = TG_PostUpdateWork;
+	PrimaryActorTick.bStartWithTickEnabled = true;
+	bAllowTickBeforeBeginPlay = true;
+
+	SetActorEnableCollision(false);
+
+	RootComponent = CreateDefaultSubobject<USceneComponent>(USceneComponent::GetDefaultSceneRootVariableName());
+
+	ViewProjectionComponent = CreateDefaultSubobject<UCompositeViewProjectionComponent>(TEXT("CameraViewProjectionComponent"));
+	AddInstanceComponent(ViewProjectionComponent);
+
+	// Default layers when creating a new actor
+	CompositeLayers.Reserve(3);
+	CompositeLayers.Add(CreateDefaultSubobject<UCompositeLayerMainRender>(TEXT("MainRenderLayer")));
+	CompositeLayers.Add(CreateDefaultSubobject<UCompositeLayerShadowReflection>(TEXT("ShadowReflectionLayer")));
+	CompositeLayers.Add(CreateDefaultSubobject<UCompositeLayerPlate>(TEXT("PlateLayer")));
+
+	// By default, we disable the shadow reflection catcher due to its significant added cost.
+	CompositeLayers[1]->bIsEnabled = false;
+}
+
+ACompositeActor::~ACompositeActor() = default;
+
+void ACompositeActor::PostRegisterAllComponents()
+{
+	using namespace UE::Composite::Private;
+
+	Super::PostRegisterAllComponents();
+
+#if WITH_EDITOR
+	if (IsTemplate() || bIsModifyingAProperty)
+	{
+		return;
+	}
+#endif
+
+	if (IsActive())
+	{
+		ApplyAllState();
+	}
+}
+
+void ACompositeActor::UnregisterAllComponents(bool bForReregister)
+{
+	using namespace UE::Composite::Private;
+
+	Super::UnregisterAllComponents(bForReregister);
+
+	if (bForReregister || bIsModifyingAProperty)
+	{
+		return;
+	}
+
+	RemoveAllState();
+}
+
+void ACompositeActor::PostUnregisterAllComponents()
+{
+	Super::PostUnregisterAllComponents();
+}
+
+void ACompositeActor::Tick(float DeltaTime)
+{
+	Super::Tick(DeltaTime);
+
+	SCOPE_CYCLE_COUNTER(STAT_Composite_Actor_Tick);
+
+	if (!IsActive())
+	{
+		return;
+	}
+
+	UCompositeCoreSubsystem* Subsystem = UWorld::GetSubsystem<UCompositeCoreSubsystem>(GetWorld());
+	if (IsValid(Subsystem))
+	{
+		Subsystem->SetPostRenderWork(UE::Composite::Private::GetPostRenderWork(*this));
+
+		UE::CompositeCore::FBuiltInRenderPassOptions Options;
+		if (bOverridesViewUserFlags)
+		{
+			Options.ViewUserFlagsOverride = ViewUserFlags;
+		}
+		Subsystem->SetBuiltInRenderPassOptions(Options);
+	}
+}
+
+#if WITH_EDITOR
+void ACompositeActor::PreEditChange(FProperty* PropertyAboutToChange)
+{
+	TOptional<TGuardValue<bool>> ModifyingPropertyGuard;
+	if (!IsActorBeingDestroyed())
+	{
+		if ((GEditor && GEditor->bIsSimulatingInEditor && GetWorld() != nullptr) || ReregisterComponentsWhenModified())
+		{
+			if (PropertyAboutToChange)
+			{
+				ModifyingPropertyGuard.Emplace(bIsModifyingAProperty, true);
+			}
+		}
+	}
+
+	Super::PreEditChange(PropertyAboutToChange);
+
+	if (PropertyAboutToChange && PropertyAboutToChange->GetFName() == GET_MEMBER_NAME_CHECKED(ThisClass, CompositeLayers))
+	{
+		PreEditCompositeLayers = CompositeLayers;
+	}
+	else
+	{
+		PreEditCompositeLayers.Reset();
+	}
+}
+
+void ACompositeActor::PostEditChangeChainProperty(struct FPropertyChangedChainEvent& PropertyChangedChainEvent)
+{
+	TGuardValue<bool> ModifyingPropertyGuard{ bIsModifyingAProperty, true };
+
+	Super::PostEditChangeChainProperty(PropertyChangedChainEvent);
+
+	const FName PropertyName = PropertyChangedChainEvent.GetPropertyName();
+
+	if (PropertyName == GET_MEMBER_NAME_CHECKED(ThisClass, bIsEnabled))
+	{
+		SetEnabled(bIsEnabled);
+	}
+	else if (PropertyName == GET_MEMBER_NAME_CHECKED(ThisClass, CompositeLayers))
+	{
+		// Verify that there isn't another container property node further down in the chain (so we don't respond to array change events from a different array!)
+		bool bWasNestedContainerChanged = false;
+		FEditPropertyChain::TDoubleLinkedListNode* Node = PropertyChangedChainEvent.PropertyChain.GetActiveNode()->GetNextNode();
+		while (!bWasNestedContainerChanged && Node && Node->GetValue())
+		{
+			bWasNestedContainerChanged = Node->GetValue()->IsA<FArrayProperty>() || Node->GetValue()->IsA<FSetProperty>() || Node->GetValue()->IsA<FMapProperty>();
+			Node = Node->GetNextNode();
+		}
+
+		if (!bWasNestedContainerChanged)
+		{
+			const int32 ArrayIndex = PropertyChangedChainEvent.GetArrayIndex(PropertyName.ToString());
+
+			int32 MainRenderIndex = INDEX_NONE;
+			
+			for (int32 Index =0; Index < CompositeLayers.Num(); ++Index)
+			{
+				const TObjectPtr<UCompositeLayerBase>& Layer = CompositeLayers[Index];
+
+				if (!IsValid(Layer) || !Layer->bIsEnabled)
+				{
+					continue;
+				}
+
+				if (Layer->IsA(UCompositeLayerMainRender::StaticClass()))
+				{
+					MainRenderIndex = Index;
+				}
+			}
+
+			if (PropertyChangedChainEvent.ChangeType == EPropertyChangeType::ArrayAdd)
+			{
+				// Always add elements at the start of the array (top of the stack)
+				Algo::Rotate(CompositeLayers, CompositeLayers.Num() - 1);
+			}
+			else if (PropertyChangedChainEvent.ChangeType == EPropertyChangeType::ValueSet || PropertyChangedChainEvent.ChangeType == EPropertyChangeType::Duplicate)
+			{
+				if (CompositeLayers.IsValidIndex(ArrayIndex))
+				{
+					if (const UCompositeLayerShadowReflection* CompositeLayer = Cast<UCompositeLayerShadowReflection>(CompositeLayers[ArrayIndex]))
+					{
+						FindOrCreateReflectionShadowCatcherPair(CompositeLayer, {});
+					}
+				}
+			}
+			else if (PropertyChangedChainEvent.ChangeType == EPropertyChangeType::ArrayRemove)
+			{
+				if (PreEditCompositeLayers.IsValidIndex(ArrayIndex))
+				{
+					if (const UCompositeLayerShadowReflection* CompositeLayer = Cast<UCompositeLayerShadowReflection>(PreEditCompositeLayers[ArrayIndex]))
+					{
+						DestroyReflectionShadowCatcherPair(CompositeLayer);
+					}
+				}
+			}
+			else if (PropertyChangedChainEvent.ChangeType == EPropertyChangeType::ArrayClear)
+			{
+				for (const UCompositeLayerBase* CompositeLayer : PreEditCompositeLayers)
+				{
+					DestroyReflectionShadowCatcherPair(CompositeLayer);
+				}
+			}
+
+			if (PropertyChangedChainEvent.PropertyChain.Num() <= 1)
+			{
+				// We only activate the holdout composite pipeline if the plate layer is below the main render
+				for (int32 Index = 0; Index < CompositeLayers.Num(); ++Index)
+				{
+					if (UCompositeLayerPlate* PlateLayer = Cast<UCompositeLayerPlate>(CompositeLayers[Index]))
+					{
+						const bool bHoldoutComposite = Index > MainRenderIndex;
+						PlateLayer->SetHoldoutComposite(bHoldoutComposite);
+					}
+				}
+			}
+		}
+	}
+	else if (PropertyName == GET_MEMBER_NAME_CHECKED(ThisClass, Camera))
+	{
+		if (IsValid(ViewProjectionComponent))
+		{
+			ViewProjectionComponent->SetTargetComponent(Camera);
+		}
+	}
+}
+
+void ACompositeActor::SetIsTemporarilyHiddenInEditor(bool bIsHidden)
+{
+	Super::SetIsTemporarilyHiddenInEditor(bIsHidden);
+
+	if (IsActive())
+	{
+		ApplyAllState();
+	}
+	else
+	{
+		RemoveAllState();
+	}
+}
+
+bool ACompositeActor::ShouldTickIfViewportsOnly() const
+{
+	return true;
+}
+#endif
+
+bool ACompositeActor::IsEnabled() const
+{
+	return bIsEnabled;
+}
+
+void ACompositeActor::SetEnabled(bool bInEnabled)
+{
+	bIsEnabled = bInEnabled;
+
+	if (IsActive())
+	{
+		ApplyAllState();
+	}
+	else
+	{
+		RemoveAllState();
+	}
+}
+
+FComponentReference& ACompositeActor::GetCamera()
+{
+	return Camera;
+}
+
+void ACompositeActor::SetCamera(const FComponentReference& InComponentReference)
+{
+	Camera = InComponentReference;
+	
+	if (IsValid(ViewProjectionComponent))
+	{
+		ViewProjectionComponent->SetTargetComponent(Camera);
+	}
+}
+
+TShadowReflectionCatcherPair ACompositeActor::FindOrCreateReflectionShadowCatcherPair(const UCompositeLayerBase* InLayer, const TArray<TObjectPtr<AActor>>& InActorsToHide)
+{
+	bool bComponentCreated = false;
+
+	if (!ReflectionShadowCatchers_CG.Contains(InLayer))
+	{
+		const FName UniqueName = MakeUniqueObjectName(this, UCompositeShadowReflectionCatcherComponent::StaticClass(), FName("ShadowReflectionCatcher_CG"));
+		UCompositeShadowReflectionCatcherComponent* Component = NewObject<UCompositeShadowReflectionCatcherComponent>(this, UniqueName, RF_Transactional);
+		Component->SetupAttachment(RootComponent);
+		Component->RegisterComponent();
+		AddInstanceComponent(Component);
+		
+		ReflectionShadowCatchers_CG.Add(InLayer, Component);
+		bComponentCreated = true;
+	}
+
+	if (!ReflectionShadowCatchers_NoCG.Contains(InLayer))
+	{
+		const FName UniqueName = MakeUniqueObjectName(this, UCompositeShadowReflectionCatcherComponent::StaticClass(), FName("ShadowReflectionCatcher_NoCG"));
+		UCompositeShadowReflectionCatcherComponent* Component = NewObject<UCompositeShadowReflectionCatcherComponent>(this, UniqueName, RF_Transactional);
+		Component->SetupAttachment(RootComponent);
+		Component->RegisterComponent();
+		AddInstanceComponent(Component);
+
+		ReflectionShadowCatchers_NoCG.Add(InLayer, Component);
+		bComponentCreated = true;
+	}
+
+#if WITH_EDITOR
+	if (bComponentCreated)
+	{
+		FModuleManager::LoadModuleChecked<FLevelEditorModule>("LevelEditor").BroadcastComponentsEdited();
+	}
+#endif
+
+	TShadowReflectionCatcherPair OutPair;
+	OutPair[0] = ReflectionShadowCatchers_CG[InLayer];
+	OutPair[1] = ReflectionShadowCatchers_NoCG[InLayer];
+	OutPair[1]->HiddenActors = InActorsToHide;
+
+	return OutPair;
+}
+
+void ACompositeActor::DestroyReflectionShadowCatcherPair(const UCompositeLayerBase* InLayer)
+{
+	bool bComponentRemoved = false;
+
+	{
+		TObjectPtr<UCompositeShadowReflectionCatcherComponent> RemovedComponent;
+		if (ReflectionShadowCatchers_CG.RemoveAndCopyValue(InLayer, RemovedComponent))
+		{
+			if (IsValid(RemovedComponent))
+			{
+				RemoveInstanceComponent(RemovedComponent);
+				RemovedComponent->UnregisterComponent();
+				RemovedComponent->DestroyComponent();
+
+				bComponentRemoved = true;
+			}
+		}
+	}
+
+	{
+		TObjectPtr<UCompositeShadowReflectionCatcherComponent> RemovedComponent;
+		if (ReflectionShadowCatchers_NoCG.RemoveAndCopyValue(InLayer, RemovedComponent))
+		{
+			if (IsValid(RemovedComponent))
+			{
+				RemoveInstanceComponent(RemovedComponent);
+				RemovedComponent->UnregisterComponent();
+				RemovedComponent->DestroyComponent();
+
+				bComponentRemoved = true;
+			}
+		}
+	}
+
+#if WITH_EDITOR
+	if (bComponentRemoved)
+	{
+		FModuleManager::LoadModuleChecked<FLevelEditorModule>("LevelEditor").BroadcastComponentsEdited();
+	}
+#endif
+}
+
+void ACompositeActor::SetReflectionShadowCatcherCaptureEveryFrame(const UCompositeLayerBase* InLayer, bool bState)
+{
+	TObjectPtr<UCompositeShadowReflectionCatcherComponent>* Component0 = ReflectionShadowCatchers_CG.Find(InLayer);
+	if (Component0)
+	{
+		(*Component0)->bCaptureEveryFrame = bState;
+	}
+
+	TObjectPtr<UCompositeShadowReflectionCatcherComponent>* Component1 = ReflectionShadowCatchers_NoCG.Find(InLayer);
+	if (Component1)
+	{
+		(*Component1)->bCaptureEveryFrame = bState;
+	}
+}
+
+bool ACompositeActor::IsActive() const
+{
+	bool bIsActive = bIsEnabled;
+
+#if WITH_EDITOR
+	if (IsTemporarilyHiddenInEditor())
+	{
+		bIsActive = false;
+	}
+#endif
+
+	return bIsActive;
+}
+
+void ACompositeActor::ApplyAllState()
+{
+	UCompositeCoreSubsystem* Subsystem = UWorld::GetSubsystem<UCompositeCoreSubsystem>(GetWorld());
+
+	OnEnabledStateChange.ExecuteIfBound(true, Subsystem);
+
+	UE::Composite::Private::SetCaptureEveryFrame(ReflectionShadowCatchers_CG, true);
+	UE::Composite::Private::SetCaptureEveryFrame(ReflectionShadowCatchers_NoCG, true);
+}
+
+void ACompositeActor::RemoveAllState()
+{
+	UCompositeCoreSubsystem* Subsystem = UWorld::GetSubsystem<UCompositeCoreSubsystem>(GetWorld());
+
+	OnEnabledStateChange.ExecuteIfBound(false, Subsystem);
+	
+	if (IsValid(Subsystem))
+	{
+		Subsystem->ResetPostRenderWork();
+		Subsystem->ResetBuiltInRenderPassOptions();
+	}
+
+	UE::Composite::Private::SetCaptureEveryFrame(ReflectionShadowCatchers_CG, false);
+	UE::Composite::Private::SetCaptureEveryFrame(ReflectionShadowCatchers_NoCG, false);
+}
+
+
diff --git a/Engine/Plugins/Experimental/Compositing/Composite/Source/Composite/Private/CompositeAssetUserData.cpp b/Engine/Plugins/Experimental/Compositing/Composite/Source/Composite/Private/CompositeAssetUserData.cpp
new file mode 100644
--- /dev/null
+++ b/Engine/Plugins/Experimental/Compositing/Composite/Source/Composite/Private/CompositeAssetUserData.cpp
@@ -0,0 +1,15 @@
+// Copyright Epic Games, Inc. All Rights Reserved.
+
+#include "CompositeAssetUserData.h"
+
+#include "Components/PrimitiveComponent.h"
+
+#include UE_INLINE_GENERATED_CPP_BY_NAME(CompositeAssetUserData)
+
+void UCompositeAssetUserData::PostEditChangeOwner(const FPropertyChangedEvent& PropertyChangedEvent)
+{
+	if (UPrimitiveComponent* Outer = GetTypedOuter<UPrimitiveComponent>())
+	{
+		OnPostEditChangeOwner.ExecuteIfBound(*Outer);
+	}
+}
diff --git a/Engine/Plugins/Experimental/Compositing/Composite/Source/Composite/Private/CompositeMeshActor.cpp b/Engine/Plugins/Experimental/Compositing/Composite/Source/Composite/Private/CompositeMeshActor.cpp
new file mode 100644
--- /dev/null
+++ b/Engine/Plugins/Experimental/Compositing/Composite/Source/Composite/Private/CompositeMeshActor.cpp
@@ -0,0 +1,15 @@
+// Copyright Epic Games, Inc. All Rights Reserved.
+
+#include "CompositeMeshActor.h"
+
+#include "Components/CompositeMeshComponent.h"
+
+ACompositeMeshActor::ACompositeMeshActor(const FObjectInitializer& ObjectInitializer)
+	: Super(ObjectInitializer)
+{
+	CompositeMeshComponent = CreateDefaultSubobject<UCompositeMeshComponent>(TEXT("DefaultCompositeMeshComponent"));
+	CompositeMeshComponent->SetupAttachment(RootComponent);
+}
+
+ACompositeMeshActor::~ACompositeMeshActor() = default;
+
diff --git a/Engine/Plugins/Experimental/Compositing/Composite/Source/Composite/Private/CompositeModule.cpp b/Engine/Plugins/Experimental/Compositing/Composite/Source/Composite/Private/CompositeModule.cpp
new file mode 100644
--- /dev/null
+++ b/Engine/Plugins/Experimental/Compositing/Composite/Source/Composite/Private/CompositeModule.cpp
@@ -0,0 +1,21 @@
+// Copyright Epic Games, Inc. All Rights Reserved.
+
+#include "CompositeModule.h"
+
+#include "Interfaces/IPluginManager.h"
+#include "Misc/Paths.h"
+#include "Modules/ModuleManager.h"
+#include "ShaderCore.h"
+
+DEFINE_LOG_CATEGORY(LogComposite);
+
+void FCompositeModule::StartupModule()
+{
+	AddShaderSourceDirectoryMapping(TEXT("/Plugin/Composite"), FPaths::Combine(IPluginManager::Get().FindPlugin(TEXT("Composite"))->GetBaseDir(), TEXT("Shaders")));
+}
+
+void FCompositeModule::ShutdownModule()
+{
+}
+
+IMPLEMENT_MODULE( FCompositeModule, Composite )
diff --git a/Engine/Plugins/Experimental/Compositing/Composite/Source/Composite/Private/Layers/CompositeLayerBase.cpp b/Engine/Plugins/Experimental/Compositing/Composite/Source/Composite/Private/Layers/CompositeLayerBase.cpp
new file mode 100644
--- /dev/null
+++ b/Engine/Plugins/Experimental/Compositing/Composite/Source/Composite/Private/Layers/CompositeLayerBase.cpp
@@ -0,0 +1,84 @@
+// Copyright Epic Games, Inc. All Rights Reserved.
+
+#include "Layers/CompositeLayerBase.h"
+
+#include "CompositeActor.h"
+
+UCompositeLayerBase::UCompositeLayerBase(const FObjectInitializer& ObjectInitializer)
+	: Super(ObjectInitializer)
+{
+}
+
+UCompositeLayerBase::~UCompositeLayerBase() = default;
+
+
+#if WITH_EDITOR
+bool UCompositeLayerBase::CanEditChange(const FProperty* InProperty) const
+{
+	bool bIsEditable = Super::CanEditChange(InProperty);
+
+	if (bIsEditable && InProperty)
+	{
+		const FName PropertyName = InProperty->GetFName();
+
+		// Always return true for enabled toggle
+		if (PropertyName != GET_MEMBER_NAME_CHECKED(ThisClass, bIsEnabled))
+		{
+			bIsEditable = bIsEnabled;
+		}
+		
+		if (PropertyName == GET_MEMBER_NAME_CHECKED(ThisClass, Operation))
+		{
+			bIsEditable = bIsEnabled && !bIsSolo;
+		}
+	}
+
+	return bIsEditable;
+}
+
+void UCompositeLayerBase::PostEditChangeProperty(FPropertyChangedEvent& PropertyChangedEvent)
+{
+	Super::PostEditChangeProperty(PropertyChangedEvent);
+
+	if (PropertyChangedEvent.GetMemberPropertyName() == GET_MEMBER_NAME_CHECKED(ThisClass, bIsSolo))
+	{
+		if (ACompositeActor* CompositeActor = GetTypedOuter<ACompositeActor>())
+		{
+			for (TObjectPtr<UCompositeLayerBase>& CompositeLayer : CompositeActor->CompositeLayers)
+			{
+				if (CompositeLayer != this)
+				{
+					CompositeLayer->bIsSolo = false;
+				}
+			}
+		}
+	}
+}
+#endif //WITH_EDITOR
+
+ECompositeCoreMergeOp UCompositeLayerBase::GetOperation(const FTraversalContext& InContext) const
+{
+	return InContext.bIsSolo ? ECompositeCoreMergeOp::None : Operation;
+}
+
+UE::CompositeCore::FPassInputDecl UCompositeLayerBase::GetDefaultSecondInput(const FTraversalContext& InContext) const
+{
+	using namespace UE::CompositeCore;
+
+	UE::CompositeCore::FPassInputDecl InputDecl;
+
+	if (InContext.bIsFirstPass)
+	{
+		// Force an empty previous texture on the first pass
+		FPassExternalResourceDesc Desc1;
+		Desc1.Id = BUILT_IN_EMPTY_ID;
+		InputDecl.Set<FPassExternalResourceDesc>(Desc1);
+	}
+	else
+	{
+		InputDecl.Set<FPassInternalResourceDesc>({});
+	}
+
+	return InputDecl;
+}
+
diff --git a/Engine/Plugins/Experimental/Compositing/Composite/Source/Composite/Private/Layers/CompositeLayerMainRender.cpp b/Engine/Plugins/Experimental/Compositing/Composite/Source/Composite/Private/Layers/CompositeLayerMainRender.cpp
new file mode 100644
--- /dev/null
+++ b/Engine/Plugins/Experimental/Compositing/Composite/Source/Composite/Private/Layers/CompositeLayerMainRender.cpp
@@ -0,0 +1,34 @@
+// Copyright Epic Games, Inc. All Rights Reserved.
+
+#include "Layers/CompositeLayerMainRender.h"
+
+#include "CompositeActor.h"
+
+UCompositeLayerMainRender::UCompositeLayerMainRender(const FObjectInitializer& ObjectInitializer)
+	: Super(ObjectInitializer)
+{
+	Operation = ECompositeCoreMergeOp::Over;
+}
+
+UCompositeLayerMainRender::~UCompositeLayerMainRender() = default;
+
+bool UCompositeLayerMainRender::GetProxy(FTraversalContext& InContext, FSceneRenderingBulkObjectAllocator& InFrameAllocator, FCompositeCorePassProxy*& OutProxy) const
+{
+	using namespace UE::CompositeCore;
+
+	FPassInputDeclArray Inputs;
+	Inputs.SetNum(2);
+
+	FPassInternalResourceDesc Desc0;
+	Desc0.Index = 0;
+	Desc0.bOriginalCopyBeforePasses = true;
+	Inputs[0].Set<FPassInternalResourceDesc>(Desc0);
+	Inputs[1] = GetDefaultSecondInput(InContext);
+
+	FMergePassProxy* Proxy = InFrameAllocator.Create<FMergePassProxy>(MoveTemp(Inputs));
+	Proxy->MergeOp = GetOperation(InContext);
+
+	OutProxy = Proxy;
+	return true;
+}
+
diff --git a/Engine/Plugins/Experimental/Compositing/Composite/Source/Composite/Private/Layers/CompositeLayerPlate.cpp b/Engine/Plugins/Experimental/Compositing/Composite/Source/Composite/Private/Layers/CompositeLayerPlate.cpp
new file mode 100644
--- /dev/null
+++ b/Engine/Plugins/Experimental/Compositing/Composite/Source/Composite/Private/Layers/CompositeLayerPlate.cpp
@@ -0,0 +1,669 @@
+// Copyright Epic Games, Inc. All Rights Reserved.
+
+#include "Layers/CompositeLayerPlate.h"
+
+#include "Components/StaticMeshComponent.h"
+#include "CompositeActor.h"
+#include "CompositeAssetUserData.h"
+#include "CompositeCoreSubsystem.h"
+#include "CompositeCoreSettings.h"
+#include "Engine/Texture2D.h"
+#include "MediaTexture.h"
+#include "Engine/TextureRenderTarget2D.h"
+#include "Materials/Material.h"
+#include "Materials/MaterialInstanceDynamic.h"
+#include "UObject/ConstructorHelpers.h"
+
+#if WITH_EDITOR
+#include "UObject/UObjectGlobals.h"
+#include "Framework/Notifications/NotificationManager.h"
+#include "Misc/TransactionObjectEvent.h"
+#include "Widgets/Notifications/SNotificationList.h"
+#endif
+
+const FLazyName UCompositeLayerPlate::CompositeTextureName = FLazyName(TEXT("CompositeTexture"));
+
+#define LOCTEXT_NAMESPACE "Composite"
+
+namespace UE
+{
+	namespace Composite
+	{
+		namespace Private
+		{
+			/** Convenience function to remove composite asset user data from tracked mesh actors (before their removal). */
+			static void RemoveCompositeAssetUserData(TConstArrayView<AActor*> InPreEditMeshes, TConstArrayView<AActor*> InPostEditMeshes)
+			{
+				for (const AActor* PreEditMesh : InPreEditMeshes)
+				{
+					if (!IsValid(PreEditMesh))
+					{
+						continue;
+					}
+
+					if (!InPostEditMeshes.Contains(PreEditMesh))
+					{
+						for (UActorComponent* Component : PreEditMesh->GetComponents())
+						{
+							UPrimitiveComponent* PrimitiveComponent = Cast<UPrimitiveComponent>(Component);
+
+							if (IsValid(PrimitiveComponent))
+							{
+								PrimitiveComponent->RemoveUserDataOfClass(UCompositeAssetUserData::StaticClass());
+							}
+						}
+					}
+				}
+			}
+
+			/** Convenience function to remove all composite asset user data from tracked mesh actors. */
+			static void RemoveCompositeAllAssetUserData(TConstArrayView<AActor*> InPreEditMeshes)
+			{
+				RemoveCompositeAssetUserData(InPreEditMeshes, {});
+			}
+		}
+	}
+}
+
+UCompositeLayerPlate::UCompositeLayerPlate(const FObjectInitializer& ObjectInitializer)
+	: Super(ObjectInitializer)
+{
+	static ConstructorHelpers::FObjectFinder<UTexture2D> DefaultMediaTexture = TEXT("/Composite/Textures/T_Composite_SMPTE_Color_Bars_16x9.T_Composite_SMPTE_Color_Bars_16x9");
+	
+	Texture = Cast<UTexture>(DefaultMediaTexture.Object);
+	bHoldoutComposite = true;
+	Operation = ECompositeCoreMergeOp::Over;
+}
+
+UCompositeLayerPlate::~UCompositeLayerPlate() = default;
+
+void UCompositeLayerPlate::Tick(float DeltaTime)
+{
+	if(!bIsEnabled)
+	{
+		return;
+	}
+
+	if (ACompositeActor* CompositeActor = GetTypedOuter<ACompositeActor>())
+	{
+		if (!CompositeActor->IsActive())
+		{
+			return;
+		}
+	}
+
+	UCompositeCoreSubsystem* Subsystem = UWorld::GetSubsystem<UCompositeCoreSubsystem>(GetWorld());
+	
+	if (IsValid(Subsystem))
+	{
+		if (UE::CompositeCore::IsRegisterPrimitivesOnTickEnabled() && bHoldoutComposite)
+		{
+			Subsystem->RegisterPrimitives(GetPrimitives());
+		}
+	}
+}
+
+ETickableTickType UCompositeLayerPlate::GetTickableTickType() const
+{
+	return HasAnyFlags(RF_ClassDefaultObject) ? ETickableTickType::Never : ETickableTickType::Always;
+}
+
+bool UCompositeLayerPlate::IsTickable() const
+{
+	return (HasAnyFlags(RF_ClassDefaultObject) == false);
+}
+
+UWorld* UCompositeLayerPlate::GetTickableGameObjectWorld() const
+{
+	return GetWorld();
+}
+
+TStatId UCompositeLayerPlate::GetStatId() const
+{
+	RETURN_QUICK_DECLARE_CYCLE_STAT(UCompositeLayerPlate, STATGROUP_Tickables);
+}
+
+void UCompositeLayerPlate::PostInitProperties()
+{
+	Super::PostInitProperties();
+
+	if (!HasAnyFlags(RF_ClassDefaultObject))
+	{
+		ACompositeActor* CompositeActor = GetTypedOuter<ACompositeActor>();
+		
+		if (ensure(IsValid(CompositeActor)))
+		{
+			CompositeActor->OnEnabledStateChange.BindUObject(this, &UCompositeLayerPlate::HandleHoldoutCompositeStateChange);
+		}
+	}
+}
+
+void UCompositeLayerPlate::PostLoad()
+{
+	Super::PostLoad();
+
+	if (!HasAnyFlags(RF_ClassDefaultObject))
+	{
+		UpdateCompositeMeshes();
+	}
+}
+
+const TArray<AActor*> UCompositeLayerPlate::GetCompositeMeshes() const
+{
+	return CompositeMeshes;
+}
+
+void UCompositeLayerPlate::SetCompositeMeshes(TArray<AActor*> InCompositeMeshes)
+{
+	// First, we untrack current meshes
+	UE::Composite::Private::RemoveCompositeAssetUserData(CompositeMeshes, InCompositeMeshes);
+
+	CompositeMeshes = MoveTemp(InCompositeMeshes);
+	
+	UpdateCompositeMeshes();
+}
+
+void UCompositeLayerPlate::SetHoldoutComposite(bool bInHoldoutComposite)
+{
+	if (bHoldoutComposite != bInHoldoutComposite)
+	{
+		bHoldoutComposite = bInHoldoutComposite;
+		
+		HandleHoldoutCompositeStateChange(bHoldoutComposite, UWorld::GetSubsystem<UCompositeCoreSubsystem>(GetWorld()));
+	}
+}
+
+UTexture* UCompositeLayerPlate::GetCompositeTexture() const
+{
+	if (IsValid(TextureProcessed) && GetValidPreprocessingPassNum() > 0)
+	{
+		return TextureProcessed.Get();
+	}
+	else if (IsValid(Texture))
+	{
+		return Texture.Get();
+	}
+
+	return nullptr;
+}
+
+#if WITH_EDITOR
+bool UCompositeLayerPlate::CanEditChange(const FProperty* InProperty) const
+{
+	bool bIsEditable = Super::CanEditChange(InProperty);
+
+	if (bIsEditable && InProperty)
+	{
+		const FName PropertyName = InProperty->GetFName();
+
+		// Always return true for enabled toggle
+		if (PropertyName == GET_MEMBER_NAME_CHECKED(ThisClass, CompositeMeshes))
+		{
+			bIsEditable = bHoldoutComposite;
+		}
+	}
+
+	return bIsEditable;
+}
+
+void UCompositeLayerPlate::PreEditChange(FProperty* PropertyThatWillChange)
+{
+	Super::PreEditChange(PropertyThatWillChange);
+
+	if (!PropertyThatWillChange)
+	{
+		return;
+	}
+
+	const FName PropertyName = PropertyThatWillChange->GetFName();
+
+	if (PropertyName == GET_MEMBER_NAME_CHECKED(ThisClass, CompositeMeshes))
+	{
+		PreEditCompositeMeshes = CompositeMeshes;
+
+		UCompositeCoreSubsystem* Subsystem = UWorld::GetSubsystem<UCompositeCoreSubsystem>(GetWorld());
+		if (IsValid(Subsystem))
+		{
+			Subsystem->UnregisterPrimitives(GetPrimitives());
+		}
+	}
+}
+
+void UCompositeLayerPlate::PostEditChangeProperty(FPropertyChangedEvent& PropertyChangedEvent)
+{
+	Super::PostEditChangeProperty(PropertyChangedEvent);
+
+	const FName PropertyName = PropertyChangedEvent.GetMemberPropertyName();
+
+	if (PropertyName == GET_MEMBER_NAME_CHECKED(ThisClass, Texture))
+	{
+		UpdateCompositeMeshes();
+	}
+	if (PropertyName == GET_MEMBER_NAME_CHECKED(ThisClass, bHoldoutComposite))
+	{
+		HandleHoldoutCompositeStateChange(bHoldoutComposite, UWorld::GetSubsystem<UCompositeCoreSubsystem>(GetWorld()));
+	}
+	else if (PropertyName == GET_MEMBER_NAME_CHECKED(ThisClass, Preprocessing))
+	{
+		const bool bChangeWasRemoval = PropertyChangedEvent.ChangeType == EPropertyChangeType::ArrayRemove
+			|| PropertyChangedEvent.ChangeType == EPropertyChangeType::ArrayClear;
+
+		if (bChangeWasRemoval)
+		{
+			if (GetValidPreprocessingPassNum() == 0)
+			{
+				TextureProcessed->ConditionalBeginDestroy();
+				TextureProcessed->ReleaseResource();
+				TextureProcessed = nullptr;
+
+				CollectGarbage(GARBAGE_COLLECTION_KEEPFLAGS);
+			}
+		}
+		else if (PropertyChangedEvent.ChangeType == EPropertyChangeType::ArrayAdd ||
+			PropertyChangedEvent.ChangeType == EPropertyChangeType::ValueSet)
+		{
+			if (IsValid(Texture) && !IsValid(TextureProcessed))
+			{
+				FIntPoint Resolution = GetTypedOuter<ACompositeActor>()->RenderResolution;
+
+#if 0 // Note: We may wish to expose options to retain the original texture resolution
+				if (const UMediaTexture* MediaTexture = Cast<UMediaTexture>(Texture))
+				{
+					Resolution = FIntPoint(MediaTexture->GetWidth(), MediaTexture->GetHeight());
+				}
+				else if (const UTexture2D* Texture2D = Cast<UTexture2D>(Texture))
+				{
+					Resolution = FIntPoint(Texture2D->GetSizeX(), Texture2D->GetSizeY());
+				}
+#endif
+
+				const FName UniqueName = MakeUniqueObjectName(this, UTextureRenderTarget2D::StaticClass(), FName("CompositeProcessedPlate"));
+				TextureProcessed = NewObject<UTextureRenderTarget2D>(this, UniqueName, RF_Transient);
+				TextureProcessed->SizeX = Resolution.X;
+				TextureProcessed->SizeY = Resolution.Y;
+				TextureProcessed->RenderTargetFormat = ETextureRenderTargetFormat::RTF_RGBA16f;
+				TextureProcessed->bForceLinearGamma = true;
+				TextureProcessed->bSupportsUAV = true;
+
+				TextureProcessed->UpdateResource();
+			}
+		}
+
+		UpdateCompositeMeshes();
+	}
+	else if (PropertyName == GET_MEMBER_NAME_CHECKED(ThisClass, CompositeMeshes))
+	{
+		const bool bChangeWasRemoval = PropertyChangedEvent.ChangeType == EPropertyChangeType::ArrayRemove
+			|| PropertyChangedEvent.ChangeType == EPropertyChangeType::ArrayClear;
+
+		if (bChangeWasRemoval)
+		{
+			UE::Composite::Private::RemoveCompositeAssetUserData(PreEditCompositeMeshes, CompositeMeshes);
+		}
+		else
+		{
+			if (PropertyChangedEvent.ChangeType == EPropertyChangeType::ArrayAdd ||
+				PropertyChangedEvent.ChangeType == EPropertyChangeType::ValueSet ||
+				PropertyChangedEvent.ChangeType == EPropertyChangeType::ResetToDefault)
+			{
+				const int32 AlteredIndex = PropertyChangedEvent.GetArrayIndex(PropertyName.ToString());
+				if (CompositeMeshes.IsValidIndex(AlteredIndex))
+				{
+					TObjectPtr<AActor>& NewCompositeMesh = CompositeMeshes[AlteredIndex];
+
+					if (IsCompositeMeshActorAlreadyInUse(NewCompositeMesh))
+					{
+						FNotificationInfo NotifyInfo(LOCTEXT("CompositeMeshAlreadyInUse", "The composite mesh is already in use by another layer."));
+						NotifyInfo.ExpireDuration = 4.0f;
+						FSlateNotificationManager::Get().AddNotification(NotifyInfo)->SetCompletionState(SNotificationItem::CS_Fail);
+
+						NewCompositeMesh = nullptr;
+					}
+				}
+			}
+
+			UpdateCompositeMeshes();
+		}
+
+		UCompositeCoreSubsystem* Subsystem = UWorld::GetSubsystem<UCompositeCoreSubsystem>(GetWorld());
+		if (IsValid(Subsystem))
+		{
+			// Register primitives to be marked as holdout and rendered in the built-in composite custom render pass.
+			Subsystem->RegisterPrimitives(GetPrimitives());
+		}
+
+		PreEditCompositeMeshes.Reset();
+	}
+}
+
+void UCompositeLayerPlate::PostTransacted(const FTransactionObjectEvent& InTransactionEvent)
+{
+	Super::PostTransacted(InTransactionEvent);
+
+	if (InTransactionEvent.HasPropertyChanges())
+	{
+		const TArray<FName>& ChangedProperties = InTransactionEvent.GetChangedProperties();
+
+		if( ChangedProperties.Contains(GET_MEMBER_NAME_CHECKED(ThisClass, Texture)))
+		{
+			UpdateCompositeMeshes();
+		}
+		else if (ChangedProperties.Contains(GET_MEMBER_NAME_CHECKED(ThisClass, CompositeMeshes)))
+		{
+			UpdateCompositeMeshes();
+
+			UCompositeCoreSubsystem* Subsystem = UWorld::GetSubsystem<UCompositeCoreSubsystem>(GetWorld());
+			if (IsValid(Subsystem))
+			{
+				// Register primitives to be marked as holdout and rendered in the built-in composite custom render pass.
+				Subsystem->RegisterPrimitives(GetPrimitives());
+			}
+		}
+		else if (ChangedProperties.Contains(GET_MEMBER_NAME_CHECKED(ThisClass, Texture)))
+		{
+			HandleHoldoutCompositeStateChange(bHoldoutComposite, UWorld::GetSubsystem<UCompositeCoreSubsystem>(GetWorld()));
+		}
+	}
+}
+#endif
+
+bool UCompositeLayerPlate::GetProxy(FTraversalContext& InContext, FSceneRenderingBulkObjectAllocator& InFrameAllocator, FCompositeCorePassProxy*& OutProxy) const
+{
+	using namespace UE::CompositeCore;
+
+	FPassInputDeclArray Inputs;
+	Inputs.SetNum(2);
+
+	TOptional<ResourceId> TextureId;
+	TOptional<ResourceId> TextureProcessedId;
+
+	if (IsValid(Texture))
+	{
+		FExternalTexture& ExternalTexture = InContext.ExternalTextures.AddDefaulted_GetRef();
+		ExternalTexture.Texture = Texture;
+		ExternalTexture.Metadata.bDistorted = true;
+
+		TextureId = EXTERNAL_RANGE_START_ID + static_cast<ResourceId>(InContext.ExternalTextures.Num() - 1);
+	}
+
+	if (TextureId.IsSet() && IsValid(TextureProcessed))
+	{
+		for (int32 Index=0; Index < Preprocessing.Num(); ++Index)
+		{
+			const TObjectPtr<UCompositePassBase>& Prepass = Preprocessing[Index];
+
+			if (!IsValid(Prepass) || !Prepass->IsActive())
+			{
+				continue;
+			}
+
+			FPassInputDecl PassInput;
+			if (Index == 0)
+			{
+				PassInput.Set<FPassExternalResourceDesc>({ TextureId.GetValue() });
+			}
+			else
+			{
+				PassInput.Set<FPassInternalResourceDesc>({});
+			}
+
+			FCompositeCorePassProxy* PrepassProxy;
+			if (Prepass->GetProxy(PassInput, InFrameAllocator, PrepassProxy))
+			{
+				const bool bIsLastPass = (Index == Preprocessing.Num() - 1);
+				if (bIsLastPass)
+				{
+					FExternalTexture& ExternalTexture = InContext.ExternalTextures.AddDefaulted_GetRef();
+					ExternalTexture.Texture = TextureProcessed;
+					ExternalTexture.Metadata.bDistorted = true;
+
+					check(!TextureProcessedId.IsSet());
+					TextureProcessedId = EXTERNAL_RANGE_START_ID + static_cast<ResourceId>(InContext.ExternalTextures.Num() - 1);
+
+					PrepassProxy->DeclareOutputOverride(TextureProcessedId.GetValue());
+				}
+
+				TArray<const FCompositeCorePassProxy*>& PrepassProxies = InContext.PreprocessingPasses.FindOrAdd(TextureId.GetValue());
+				PrepassProxies.Add(PrepassProxy);
+			}
+		}
+	}
+	
+	ResourceId InputId = BUILT_IN_EMPTY_ID;
+	if (bHoldoutComposite && (GetValidPrimitivesNum() > 0))
+	{
+		InputId = BUILT_IN_CRP_ID;
+	}
+	else if (TextureProcessedId.IsSet())
+	{
+		InputId = TextureProcessedId.GetValue();
+	}
+	else if (TextureId.IsSet())
+	{
+		InputId = TextureId.GetValue();
+	}
+
+	Inputs[0].Set<FPassExternalResourceDesc>({ InputId });
+	Inputs[1] = GetDefaultSecondInput(InContext);
+
+	// Copy for subpasses before we move the declared inputs
+	FPassInputDecl BaseInput0 = Inputs[0];
+
+	FMergePassProxy* Proxy = InFrameAllocator.Create<FMergePassProxy>(MoveTemp(Inputs));
+	Proxy->MergeOp = GetOperation(InContext);
+	Proxy->bUseLensDistortion = false;
+
+	if (!Postprocessing.IsEmpty())
+	{
+		TArray<const FCompositeCorePassProxy*> SubPassProxies;
+		SubPassProxies.Reserve(Postprocessing.Num());
+
+		// Iterate backwards so that lower passes in the UI are executed first
+		for (int32 PassIndex = Postprocessing.Num() - 1; PassIndex >= 0; --PassIndex)
+		{
+			const TObjectPtr<UCompositePassBase>& SubPass = Postprocessing[PassIndex];
+
+			if (IsValid(SubPass) && SubPass->IsActive())
+			{
+				FCompositeCorePassProxy* SubPassProxy;
+				if (SubPass->GetProxy(BaseInput0, InFrameAllocator, SubPassProxy))
+				{
+					SubPassProxies.Add(SubPassProxy);
+				}
+			}
+		}
+
+		// Apply sub-passes overriding the first input
+		Proxy->AddSubPasses(0, SubPassProxies);
+	}
+
+	OutProxy = Proxy;
+	return true;
+}
+
+TArray<UPrimitiveComponent*> UCompositeLayerPlate::GetPrimitives() const
+{
+	const UCompositeCorePluginSettings* Settings = GetDefault<UCompositeCorePluginSettings>();
+	check(Settings);
+
+	TArray<UPrimitiveComponent*> OutPrimitiveComponents;
+
+	for (const AActor* Actor : CompositeMeshes)
+	{
+		TInlineComponentArray<UPrimitiveComponent*> PrimComponents(Actor);
+
+		for (UPrimitiveComponent* PrimitiveComponent : PrimComponents)
+		{
+			if (Settings->IsAllowedPrimitiveClass(PrimitiveComponent))
+			{
+				OutPrimitiveComponents.Add(PrimitiveComponent);
+			}
+		}
+	}
+
+	return OutPrimitiveComponents;
+}
+
+int32 UCompositeLayerPlate::GetValidPrimitivesNum() const
+{
+	const UCompositeCorePluginSettings* Settings = GetDefault<UCompositeCorePluginSettings>();
+	check(Settings);
+
+	int32 Count = 0;
+
+	for (const AActor* Actor : CompositeMeshes)
+	{
+		TInlineComponentArray<UPrimitiveComponent*> PrimComponents(Actor);
+
+		for (UPrimitiveComponent* PrimitiveComponent : PrimComponents)
+		{
+			if (Settings->IsAllowedPrimitiveClass(PrimitiveComponent))
+			{
+				++Count;
+			}
+		}
+	}
+
+	return Count;
+}
+
+int32 UCompositeLayerPlate::GetValidPreprocessingPassNum() const
+{
+	int32 Count = 0;
+
+	for (int32 Index = 0; Index < Preprocessing.Num(); ++Index)
+	{
+		const TObjectPtr<UCompositePassBase>& Prepass = Preprocessing[Index];
+
+		if (IsValid(Prepass) && Prepass->IsActive())
+		{
+			++Count;
+		}
+	}
+
+	return Count;
+}
+
+void UCompositeLayerPlate::UpdateCompositeMeshes(bool bRegisterOnly) const
+{
+	for (AActor* MeshActor : CompositeMeshes)
+	{
+		if (!IsValid(MeshActor))
+		{
+			continue;
+		}
+
+		for (UActorComponent* Component : MeshActor->GetComponents())
+		{
+			UPrimitiveComponent* PrimitiveComponent = Cast<UPrimitiveComponent>(Component);
+
+			if (!IsValid(PrimitiveComponent))
+			{
+				continue;
+			}
+
+			// Track meshes so that we can update their material upon changes
+			if (!PrimitiveComponent->HasAssetUserDataOfClass(UCompositeAssetUserData::StaticClass()))
+			{
+				UCompositeAssetUserData* AssetUserData = NewObject<UCompositeAssetUserData>(PrimitiveComponent);
+				AssetUserData->OnPostEditChangeOwner.BindUObject(this, &UCompositeLayerPlate::UpdatePrimitiveComponent);
+				PrimitiveComponent->AddAssetUserData(AssetUserData);
+			}
+
+			if (!bRegisterOnly)
+			{
+				UpdatePrimitiveComponent(*PrimitiveComponent);
+			}
+		}
+	}
+}
+
+void UCompositeLayerPlate::UpdatePrimitiveComponent(UPrimitiveComponent& InPrimitiveComponent) const
+{
+	for (int32 ElementIndex = 0; ElementIndex < InPrimitiveComponent.GetNumMaterials(); ++ElementIndex)
+	{
+		UMaterialInterface* MaterialInterface = InPrimitiveComponent.GetMaterial(ElementIndex);
+
+		if (!IsValid(MaterialInterface))
+		{
+			continue;
+		}
+
+		TArray<FMaterialParameterInfo> ParameterInfo;
+		TArray<FGuid> ParameterIds;
+		MaterialInterface->GetAllTextureParameterInfo(ParameterInfo, ParameterIds);
+
+		// Avoid conversion to MID if parameter doesn't exist
+		const bool bHasTextureParameter = ParameterInfo.ContainsByPredicate([](const FMaterialParameterInfo& ParamInfo)
+			{
+				return ParamInfo.Name.IsEqual(CompositeTextureName);
+			}
+		);
+
+		if (bHasTextureParameter)
+		{
+			UMaterialInstanceDynamic* MID = Cast<UMaterialInstanceDynamic>(MaterialInterface);
+
+			if (!MID)
+			{
+				/**
+				* Note: Unfortunately, MIDs are the only currently viable approach to bind textures to composite meshes.
+				* Texture collection could provide a nice alternative, but are experimental, have limited support and do
+				* not provide the same world-unique instance functionality that MPCs do, which would cause dynamic
+				* bindings to dirty a texture collection asset.
+				*/
+				MID = InPrimitiveComponent.CreateAndSetMaterialInstanceDynamic(ElementIndex);
+			}
+
+			MID->SetTextureParameterValue(TEXT("CompositeTexture"), GetCompositeTexture());
+		}
+	}
+}
+
+void UCompositeLayerPlate::HandleHoldoutCompositeStateChange(bool bInEnabled, UCompositeCoreSubsystem* InSubsystem) const
+{
+	if (bInEnabled)
+	{
+		constexpr bool bRegisterOnly = true;
+		UpdateCompositeMeshes(bRegisterOnly);
+
+		if (IsValid(InSubsystem))
+		{
+			InSubsystem->RegisterPrimitives(GetPrimitives());
+		}
+	}
+	else
+	{
+		UE::Composite::Private::RemoveCompositeAllAssetUserData(CompositeMeshes);
+
+		if (IsValid(InSubsystem))
+		{
+			InSubsystem->UnregisterPrimitives(GetPrimitives());
+		}
+	}
+}
+
+bool UCompositeLayerPlate::IsCompositeMeshActorAlreadyInUse(AActor* InCompositeMeshActor) const
+{
+	ACompositeActor* CompositeActor = GetTypedOuter<ACompositeActor>();
+	if (!IsValid(CompositeActor))
+	{
+		return false;
+	}
+
+	for (const TObjectPtr<UCompositeLayerBase>& Layer : CompositeActor->CompositeLayers)
+	{
+		const UCompositeLayerPlate* PlateLayer = Cast<UCompositeLayerPlate>(Layer.Get());
+		if (IsValid(PlateLayer) && PlateLayer != this)
+		{
+			if (PlateLayer->CompositeMeshes.Contains(InCompositeMeshActor))
+			{
+				return true;
+			}
+		}
+	}
+
+	return false;
+}
+
+#undef LOCTEXT_NAMESPACE
diff --git a/Engine/Plugins/Experimental/Compositing/Composite/Source/Composite/Private/Layers/CompositeLayerShadowReflection.cpp b/Engine/Plugins/Experimental/Compositing/Composite/Source/Composite/Private/Layers/CompositeLayerShadowReflection.cpp
new file mode 100644
--- /dev/null
+++ b/Engine/Plugins/Experimental/Compositing/Composite/Source/Composite/Private/Layers/CompositeLayerShadowReflection.cpp
@@ -0,0 +1,98 @@
+// Copyright Epic Games, Inc. All Rights Reserved.
+
+#include "Layers/CompositeLayerShadowReflection.h"
+
+#include "CompositeActor.h"
+#include "Components/CompositeShadowReflectionCatcherComponent.h"
+
+UCompositeLayerShadowReflection::UCompositeLayerShadowReflection(const FObjectInitializer& ObjectInitializer)
+	: Super(ObjectInitializer)
+{
+	Operation = ECompositeCoreMergeOp::Multiply;
+}
+
+UCompositeLayerShadowReflection::~UCompositeLayerShadowReflection() = default;
+
+#if WITH_EDITOR
+void UCompositeLayerShadowReflection::PostEditChangeProperty(FPropertyChangedEvent& PropertyChangedEvent)
+{
+	Super::PostEditChangeProperty(PropertyChangedEvent);
+
+	if (PropertyChangedEvent.GetMemberPropertyName() == GET_MEMBER_NAME_CHECKED(ThisClass, bIsEnabled))
+	{
+		ACompositeActor* CompositeActor = GetTypedOuter<ACompositeActor>();
+		if (IsValid(CompositeActor))
+		{
+			CompositeActor->SetReflectionShadowCatcherCaptureEveryFrame(this, bIsEnabled);
+		}
+	}
+}
+#endif
+
+bool UCompositeLayerShadowReflection::GetProxy(FTraversalContext& InContext, FSceneRenderingBulkObjectAllocator& InFrameAllocator, FCompositeCorePassProxy*& OutProxy) const
+{
+	using namespace UE::CompositeCore;
+	constexpr int32 NumInputs = 2;
+
+	ACompositeActor* CompositeActor = GetTypedOuter<ACompositeActor>();
+	if (!IsValid(CompositeActor))
+	{
+		return false;
+	}
+
+	const TShadowReflectionCatcherPair ShadowReflectionPair = CompositeActor->FindOrCreateReflectionShadowCatcherPair(this, Actors);
+	FMergePassProxy* Proxy = nullptr;
+	FMergePassProxy* DivisionProxy = nullptr;
+
+	bool bValidComponents = true;
+
+	for (const TObjectPtr<UCompositeShadowReflectionCatcherComponent>& Component : ShadowReflectionPair)
+	{
+		bValidComponents &= IsValid(Component);
+	}
+
+	if(bValidComponents)
+	{
+		FPassInputDeclArray Inputs;
+		Inputs.SetNum(NumInputs);
+
+		for (int32 Index = 0; Index < NumInputs; ++Index)
+		{
+			const TObjectPtr<UCompositeShadowReflectionCatcherComponent>& Component = ShadowReflectionPair[Index];
+			const FIntPoint& RenderResolution = CompositeActor->RenderResolution;
+
+			FExternalTexture& ExternalTexture = InContext.ExternalTextures.AddDefaulted_GetRef();
+			ExternalTexture.Texture = Component->FindOrCreateRenderTarget(RenderResolution.X, RenderResolution.Y);
+			ExternalTexture.Metadata.bInvertedAlpha = true;
+			ExternalTexture.Metadata.bDistorted = false;
+			ExternalTexture.Metadata.Encoding = UE::CompositeCore::EEncoding::Linear;
+
+			int32 TexIndex = InContext.ExternalTextures.Num() - 1;
+			Inputs[Index].Set<FPassExternalResourceDesc>({ EXTERNAL_RANGE_START_ID + TexIndex });
+		}
+
+		DivisionProxy = InFrameAllocator.Create<FMergePassProxy>(MoveTemp(Inputs));
+		DivisionProxy->MergeOp = ECompositeCoreMergeOp::Divide;
+		DivisionProxy->bUseLensDistortion = false;
+	}
+	else
+	{
+		return false;
+	}
+
+	FPassInputDeclArray Inputs;
+	Inputs.SetNum(NumInputs);
+	// Input 0 to be replaced with shadow reflection
+	Inputs[1] = GetDefaultSecondInput(InContext);
+
+	Proxy = InFrameAllocator.Create<FMergePassProxy>(MoveTemp(Inputs));
+	Proxy->MergeOp = GetOperation(InContext);
+	if (DivisionProxy)
+	{
+		Proxy->AddSubPasses(0, { DivisionProxy });
+	}
+
+	OutProxy = Proxy;
+	return true;
+}
+
diff --git a/Engine/Plugins/Experimental/Compositing/Composite/Source/Composite/Private/Passes/CompositePassBase.cpp b/Engine/Plugins/Experimental/Compositing/Composite/Source/Composite/Private/Passes/CompositePassBase.cpp
new file mode 100644
--- /dev/null
+++ b/Engine/Plugins/Experimental/Compositing/Composite/Source/Composite/Private/Passes/CompositePassBase.cpp
@@ -0,0 +1,34 @@
+// Copyright Epic Games, Inc. All Rights Reserved.
+
+#include "Passes/CompositePassBase.h"
+
+#include "CompositeActor.h"
+
+UCompositePassBase::UCompositePassBase(const FObjectInitializer& ObjectInitializer)
+	: Super(ObjectInitializer)
+{
+
+}
+
+UCompositePassBase::~UCompositePassBase() = default;
+
+#if WITH_EDITOR
+bool UCompositePassBase::CanEditChange(const FProperty* InProperty) const
+{
+	bool bIsEditable = Super::CanEditChange(InProperty);
+
+	if (bIsEditable && InProperty)
+	{
+		const FName PropertyName = InProperty->GetFName();
+
+		// Always return true for enabled toggle
+		if (PropertyName != GET_MEMBER_NAME_CHECKED(ThisClass, bIsEnabled))
+		{
+			bIsEditable = bIsEnabled;
+		}
+	}
+
+	return bIsEditable;
+}
+#endif //WITH_EDITOR
+
diff --git a/Engine/Plugins/Experimental/Compositing/Composite/Source/Composite/Private/Passes/CompositePassCenteredScale.cpp b/Engine/Plugins/Experimental/Compositing/Composite/Source/Composite/Private/Passes/CompositePassCenteredScale.cpp
new file mode 100644
--- /dev/null
+++ b/Engine/Plugins/Experimental/Compositing/Composite/Source/Composite/Private/Passes/CompositePassCenteredScale.cpp
@@ -0,0 +1,208 @@
+// Copyright Epic Games, Inc. All Rights Reserved.
+
+#include "Passes/CompositePassCenteredScale.h"
+
+#include "Camera/CameraComponent.h"
+#include "CompositeActor.h"
+#include "Engine/Texture2D.h"
+#include "Layers/CompositeLayerPlate.h"
+#include "MediaTexture.h"
+#include "Passes/CompositeCorePassProxy.h"
+
+DECLARE_GPU_STAT_NAMED(FCompositeCenteredScale, TEXT("Composite.CenteredScale"));
+
+class FCompositeCenteredScaleShader : public FGlobalShader
+{
+	DECLARE_GLOBAL_SHADER(FCompositeCenteredScaleShader);
+	SHADER_USE_PARAMETER_STRUCT(FCompositeCenteredScaleShader, FGlobalShader);
+
+	BEGIN_SHADER_PARAMETER_STRUCT(FParameters, )
+		SHADER_PARAMETER_RDG_TEXTURE(Texture2D, InputTexture)
+		SHADER_PARAMETER_SAMPLER(SamplerState, InputSampler)
+		SHADER_PARAMETER(FVector2f, ScaleUV)
+		RENDER_TARGET_BINDING_SLOTS()
+	END_SHADER_PARAMETER_STRUCT()
+};
+IMPLEMENT_GLOBAL_SHADER(FCompositeCenteredScaleShader, "/Plugin/Composite/Private/CompositeCenteredScale.usf", "MainPS", SF_Pixel);
+
+
+namespace UE
+{
+	namespace CompositeCore
+	{
+		class FCenteredScalePassProxy : public FCompositeCorePassProxy
+		{
+		public:
+			using FCompositeCorePassProxy::FCompositeCorePassProxy;
+
+			FPassOutput Add(FRDGBuilder& GraphBuilder, const FSceneView& InView, const FPassInputArray& Inputs, const FPassContext& PassContext) const override
+			{
+				RDG_EVENT_SCOPE_STAT(GraphBuilder, FCompositeCenteredScale, "Composite.CenteredScale");
+				RDG_GPU_STAT_SCOPE(GraphBuilder, FCompositeCenteredScale);
+
+				const UE::CompositeCore::FResourceMetadata& Metadata = Inputs[0].Metadata;
+				const FScreenPassTexture& Input = Inputs[0].Texture;
+				FScreenPassRenderTarget Output = Inputs.OverrideOutput;
+
+				// If the override output is provided, it means that this is the last pass in post processing.
+				if (!Output.IsValid())
+				{
+					Output = CreateOutputRenderTarget(GraphBuilder, InView, PassContext.OutputViewRect, Input.Texture->Desc, TEXT("CenteredScaleCompositePass"));
+				}
+
+				FGlobalShaderMap* GlobalShaderMap = GetGlobalShaderMap(InView.GetFeatureLevel());
+				FScreenPassTextureViewport Viewport = FScreenPassTextureViewport(Output);
+
+				FCompositeCenteredScaleShader::FParameters* Parameters = GraphBuilder.AllocParameters<FCompositeCenteredScaleShader::FParameters>();
+				Parameters->InputTexture = Input.Texture;
+				Parameters->InputSampler = TStaticSamplerState<SF_Bilinear, AM_Border, AM_Border, AM_Clamp>::GetRHI();
+				Parameters->ScaleUV = ScaleUV;
+				Parameters->RenderTargets[0] = Output.GetRenderTargetBinding();
+
+				TShaderMapRef<FCompositeCenteredScaleShader> PixelShader(GlobalShaderMap);
+				AddDrawScreenPass(
+					GraphBuilder,
+					RDG_EVENT_NAME("Composite.CenteredScale (%dx%d)", Viewport.Extent.X, Viewport.Extent.Y),
+					InView,
+					Viewport,
+					Viewport,
+					PixelShader,
+					Parameters
+				);
+
+				return UE::CompositeCore::FPassOutput{ MoveTemp(Output), Metadata, PassOutputOverride };
+			}
+
+			FVector2f ScaleUV = FVector2f::One();
+		};
+	}
+}
+
+UCompositePassCenteredScale::UCompositePassCenteredScale(const FObjectInitializer& ObjectInitializer)
+	: Super(ObjectInitializer)
+	, ScaleMode(ECompositePassScaleMode::Automatic)
+	, SourceAspectRatio(16.0f, 9.0f)
+	, TargetAspectRatio(16.0f, 9.0f)
+	, ScaleFactor(1.0f, 1.0f)
+	, OverscanUncropMode(ECompositePassOverscanUncropMode::Automatic)
+	, Overscan(0.0f)
+{
+}
+
+UCompositePassCenteredScale::~UCompositePassCenteredScale() = default;
+
+bool UCompositePassCenteredScale::GetProxy(const UE::CompositeCore::FPassInputDecl& InputDecl, FSceneRenderingBulkObjectAllocator& InFrameAllocator, FCompositeCorePassProxy*& OutProxy) const
+{
+	using namespace UE::CompositeCore;
+
+	// Skip pass if there is no active scaling
+	if (!IsActive())
+	{
+		return false;
+	}
+
+	FCenteredScalePassProxy* Proxy = InFrameAllocator.Create<FCenteredScalePassProxy>(FPassInputDeclArray{ InputDecl });
+	Proxy->ScaleUV = CalculateScale();
+
+	OutProxy = Proxy;
+	return true;
+}
+
+bool UCompositePassCenteredScale::IsActive() const
+{
+	return bIsEnabled && !CalculateScale().Equals(FVector2f::One(), UE_SMALL_NUMBER);
+}
+
+FVector2f UCompositePassCenteredScale::CalculateScale() const
+{
+	// TODO: The scale UV could be cached, though special care would be needed to acocunt for dependent property changes
+
+	ACompositeActor* CompositeActor = GetTypedOuter<ACompositeActor>();
+	UCameraComponent* CameraComponent = nullptr;
+
+	if (CompositeActor)
+	{
+		CameraComponent = Cast<UCameraComponent>(CompositeActor->GetCamera().GetComponent(nullptr));
+	}
+
+	FVector2f ScaleUV = FVector2f::One();
+
+	auto ApplyAspectRatioScaleFn = [&ScaleUV](float SourceAR, float TargetAR)
+		{
+			if (SourceAR > TargetAR)
+			{
+				ScaleUV.X = (SourceAR != 0.0f) ? TargetAR / SourceAR : 0.0f;
+			}
+			else if (SourceAR < TargetAR)
+			{
+				ScaleUV.Y = (TargetAR != 0.0f) ? SourceAR / TargetAR : 0.0f;
+			}
+		};
+
+	switch (ScaleMode)
+	{
+	case ECompositePassScaleMode::None:
+		break;
+	case ECompositePassScaleMode::Automatic:
+	{
+		float SourceAR = 16.0f / 9.0f;
+
+		if (UCompositeLayerPlate* PlateLayer = GetTypedOuter<UCompositeLayerPlate>())
+		{
+			if (IsValid(PlateLayer->Texture))
+			{
+				if (const UMediaTexture* MediaTexture = Cast<UMediaTexture>(PlateLayer->Texture))
+				{
+					SourceAR = MediaTexture->CurrentAspectRatio;
+				}
+				else if (const UTexture2D* Texture2D = Cast<UTexture2D>(PlateLayer->Texture))
+				{
+					SourceAR = Texture2D->GetSizeX() / static_cast<float>(Texture2D->GetSizeY());
+				}
+			}
+		}
+
+		const float TargetAR = CameraComponent ? CameraComponent->AspectRatio : SourceAR;
+
+		ApplyAspectRatioScaleFn(SourceAR, TargetAR);
+	}
+	break;
+	case ECompositePassScaleMode::AspectRatio:
+	{
+		const float SourceAR = SourceAspectRatio.X / SourceAspectRatio.Y;
+		const float TargetAR = TargetAspectRatio.X / TargetAspectRatio.Y;
+
+		ApplyAspectRatioScaleFn(SourceAR, TargetAR);
+	}
+	break;
+	case ECompositePassScaleMode::Manual:
+		ScaleUV /= ScaleFactor;
+		break;
+
+	default:
+		checkNoEntry();
+		break;
+	};
+
+	switch (OverscanUncropMode)
+	{
+	case ECompositePassOverscanUncropMode::None:
+		break;
+	case ECompositePassOverscanUncropMode::Automatic:
+		if (CameraComponent)
+		{
+			ScaleUV *= 1.0f + CameraComponent->Overscan;
+		}
+		break;
+	case ECompositePassOverscanUncropMode::Manual:
+		ScaleUV *= 1.0f + Overscan;
+		break;
+
+	default:
+		checkNoEntry();
+		break;
+	}
+
+	return ScaleUV;
+}
+
diff --git a/Engine/Plugins/Experimental/Compositing/Composite/Source/Composite/Private/Passes/CompositePassFXAA.cpp b/Engine/Plugins/Experimental/Compositing/Composite/Source/Composite/Private/Passes/CompositePassFXAA.cpp
new file mode 100644
--- /dev/null
+++ b/Engine/Plugins/Experimental/Compositing/Composite/Source/Composite/Private/Passes/CompositePassFXAA.cpp
@@ -0,0 +1,25 @@
+// Copyright Epic Games, Inc. All Rights Reserved.
+
+#include "Passes/CompositePassFXAA.h"
+
+#include "Passes/CompositeCorePassFXAAProxy.h"
+
+UCompositePassFXAA::UCompositePassFXAA(const FObjectInitializer& ObjectInitializer)
+	: Super(ObjectInitializer)
+	, Quality(5)
+{
+}
+
+UCompositePassFXAA::~UCompositePassFXAA() = default;
+
+bool UCompositePassFXAA::GetProxy(const UE::CompositeCore::FPassInputDecl& InputDecl, FSceneRenderingBulkObjectAllocator& InFrameAllocator, FCompositeCorePassProxy*& OutProxy) const
+{
+	using namespace UE::CompositeCore;
+
+	FFXAAPassProxy* Proxy = InFrameAllocator.Create<FFXAAPassProxy>(UE::CompositeCore::FPassInputDeclArray{ InputDecl });
+	Proxy->QualityOverride = Quality;
+
+	OutProxy = Proxy;
+	return true;
+}
+
diff --git a/Engine/Plugins/Experimental/Compositing/Composite/Source/Composite/Private/Passes/CompositePassMaterial.cpp b/Engine/Plugins/Experimental/Compositing/Composite/Source/Composite/Private/Passes/CompositePassMaterial.cpp
new file mode 100644
--- /dev/null
+++ b/Engine/Plugins/Experimental/Compositing/Composite/Source/Composite/Private/Passes/CompositePassMaterial.cpp
@@ -0,0 +1,62 @@
+// Copyright Epic Games, Inc. All Rights Reserved.
+
+#include "Passes/CompositePassMaterial.h"
+
+#include "Materials/MaterialInterface.h"
+#include "RenderGraphBuilder.h"
+#include "RenderGraphUtils.h"
+#include "SceneView.h"
+#include "ScreenPass.h"
+#include "PostProcess/PostProcessMaterialInputs.h"
+
+class FCompositePassMaterialProxy : public FCompositeCorePassProxy
+{
+public:
+	using FCompositeCorePassProxy::FCompositeCorePassProxy;
+
+	/* Post-process material weak pointer. */
+	TWeakObjectPtr<UMaterialInterface> Material;
+
+	UE::CompositeCore::FPassOutput Add(FRDGBuilder& GraphBuilder, const FSceneView& InView, const UE::CompositeCore::FPassInputArray& Inputs, const UE::CompositeCore::FPassContext& PassContext) const override
+	{
+		const UE::CompositeCore::FResourceMetadata& Metadata = Inputs[0].Metadata;
+		FPostProcessMaterialInputs PostInputs = Inputs.ToPostProcessInputs(GraphBuilder, PassContext.SceneTextures);
+		// Specify the scene color output format in case a render target needs to be created and our input is render target incompatible.
+		PostInputs.OutputFormat = FSceneTexturesConfig::Get().ColorFormat;
+
+		constexpr bool bEvenIfPendingKill = false;
+		constexpr bool bThreadsafeTest = true;
+		if (Material.IsValid(bEvenIfPendingKill, bThreadsafeTest))
+		{
+			FScreenPassTexture Output = AddPostProcessMaterialPass(GraphBuilder, InView, PostInputs, Material.GetEvenIfUnreachable());
+
+			return UE::CompositeCore::FPassOutput{ MoveTemp(Output), Metadata, PassOutputOverride };
+		}
+
+		return UE::CompositeCore::FPassOutput{ PostInputs.ReturnUntouchedSceneColorForPostProcessing(GraphBuilder), Metadata, PassOutputOverride };
+	}
+};
+
+
+UCompositePassMaterial::UCompositePassMaterial(const FObjectInitializer& ObjectInitializer)
+	: Super(ObjectInitializer)
+{
+
+}
+
+UCompositePassMaterial::~UCompositePassMaterial() = default;
+
+bool UCompositePassMaterial::GetProxy(const UE::CompositeCore::FPassInputDecl& InputDecl, FSceneRenderingBulkObjectAllocator& InFrameAllocator, FCompositeCorePassProxy*& OutProxy) const
+{
+	if (IsValid(PostProcessMaterial) && PostProcessMaterial->IsPostProcessMaterial())
+	{
+		FCompositePassMaterialProxy* Proxy = InFrameAllocator.Create<FCompositePassMaterialProxy>(UE::CompositeCore::FPassInputDeclArray{ InputDecl });
+		Proxy->Material = PostProcessMaterial;
+		
+		OutProxy = Proxy;
+		return true;
+	}
+
+	return false;
+}
+
diff --git a/Engine/Plugins/Experimental/Compositing/Composite/Source/Composite/Private/Passes/CompositePassOpenColorIO.cpp b/Engine/Plugins/Experimental/Compositing/Composite/Source/Composite/Private/Passes/CompositePassOpenColorIO.cpp
new file mode 100644
--- /dev/null
+++ b/Engine/Plugins/Experimental/Compositing/Composite/Source/Composite/Private/Passes/CompositePassOpenColorIO.cpp
@@ -0,0 +1,63 @@
+// Copyright Epic Games, Inc. All Rights Reserved.
+
+#include "Passes/CompositePassOpenColorIO.h"
+
+#include "OpenColorIORendering.h"
+#include "ScreenPass.h"
+#include "PostProcess/PostProcessMaterialInputs.h"
+
+class FOpenColorIOCompositePassProxy
+	: public FCompositeCorePassProxy
+	, public TSharedFromThis<FOpenColorIOCompositePassProxy, ESPMode::ThreadSafe>
+{
+public:
+	FOpenColorIOCompositePassProxy(UE::CompositeCore::FPassInputDeclArray InPassDeclaredInputs, FOpenColorIORenderPassResources&& InPassResources)
+		: FCompositeCorePassProxy(MoveTemp(InPassDeclaredInputs))
+		, PassResources(InPassResources)
+	{
+
+	}
+
+	UE::CompositeCore::FPassOutput Add(FRDGBuilder& GraphBuilder, const FSceneView& InView, const UE::CompositeCore::FPassInputArray& Inputs, const UE::CompositeCore::FPassContext& PassContext) const override
+	{
+		const UE::CompositeCore::FResourceMetadata& Metadata = Inputs[0].Metadata;
+		const FScreenPassTexture& Input = Inputs[0].Texture;
+		FScreenPassRenderTarget Output = Inputs.OverrideOutput;
+
+		// If the override output is provided, it means that this is the last pass in post processing.
+		if (!Output.IsValid())
+		{
+			Output = CreateOutputRenderTarget(GraphBuilder, InView, PassContext.OutputViewRect, Input.Texture->Desc, TEXT("OcioCompositePass"));
+		}
+
+		FOpenColorIORendering::AddPass_RenderThread(
+			GraphBuilder,
+			InView,
+			Input,
+			Output,
+			PassResources
+		);
+
+		return UE::CompositeCore::FPassOutput{ MoveTemp(Output), Metadata, PassOutputOverride };
+	}
+
+	FOpenColorIORenderPassResources PassResources;
+};
+
+UCompositePassOpenColorIO::UCompositePassOpenColorIO(const FObjectInitializer& ObjectInitializer)
+	: Super(ObjectInitializer)
+{
+}
+
+UCompositePassOpenColorIO::~UCompositePassOpenColorIO() = default;
+
+bool UCompositePassOpenColorIO::GetProxy(const UE::CompositeCore::FPassInputDecl& InputDecl, FSceneRenderingBulkObjectAllocator& InFrameAllocator, FCompositeCorePassProxy*& OutProxy) const
+{
+	FOpenColorIOCompositePassProxy* Proxy = InFrameAllocator.Create<FOpenColorIOCompositePassProxy>(UE::CompositeCore::FPassInputDeclArray{ InputDecl },
+		FOpenColorIORendering::GetRenderPassResources(ColorConversionSettings, GMaxRHIFeatureLevel)
+	);
+
+	OutProxy = Proxy;
+	return true;
+}
+
diff --git a/Engine/Plugins/Experimental/Compositing/Composite/Source/Composite/Public/Components/CompositeMeshComponent.h b/Engine/Plugins/Experimental/Compositing/Composite/Source/Composite/Public/Components/CompositeMeshComponent.h
new file mode 100644
--- /dev/null
+++ b/Engine/Plugins/Experimental/Compositing/Composite/Source/Composite/Public/Components/CompositeMeshComponent.h
@@ -0,0 +1,36 @@
+// Copyright Epic Games, Inc. All Rights Reserved.
+
+#pragma once
+
+#include "CoreTypes.h"
+#include "Components/StaticMeshComponent.h"
+
+#include "CompositeMeshComponent.generated.h"
+
+#define UE_API COMPOSITE_API
+
+UCLASS(MinimalAPI, Blueprintable, ClassGroup = Composite, EditInlineNew)
+class UCompositeMeshComponent : public UStaticMeshComponent
+{
+	GENERATED_UCLASS_BODY()
+
+public:
+	//~ Begin UObject interface
+	UE_API virtual void PostInitProperties() override;
+	UE_API virtual void PostDuplicate(bool bDuplicateForPIE) override;
+	UE_API virtual void PostEditImport() override;
+	UE_API virtual void PostLoad() override;
+	//~ End UObject interface
+
+	//~ Begin UActorComponent interface
+	UE_API virtual void OnRegister() override;
+	UE_API virtual void OnUnregister() override;
+	//~ End UActorComponent interface
+
+#if WITH_EDITOR
+	UE_API virtual void PostEditChangeProperty(struct FPropertyChangedEvent& PropertyChangedEvent) override;
+#endif //WITH_EDITOR
+};
+
+#undef UE_API
+
diff --git a/Engine/Plugins/Experimental/Compositing/Composite/Source/Composite/Public/Components/CompositeShadowReflectionCatcherComponent.h b/Engine/Plugins/Experimental/Compositing/Composite/Source/Composite/Public/Components/CompositeShadowReflectionCatcherComponent.h
new file mode 100644
--- /dev/null
+++ b/Engine/Plugins/Experimental/Compositing/Composite/Source/Composite/Public/Components/CompositeShadowReflectionCatcherComponent.h
@@ -0,0 +1,28 @@
+// Copyright Epic Games, Inc. All Rights Reserved.
+
+#pragma once
+
+#include "CoreTypes.h"
+#include "Components/SceneCaptureComponent2D.h"
+
+#include "CompositeShadowReflectionCatcherComponent.generated.h"
+
+#define UE_API COMPOSITE_API
+
+UCLASS(MinimalAPI, ClassGroup = Composite, EditInlineNew, meta = (BlueprintSpawnableComponent))
+class UCompositeShadowReflectionCatcherComponent : public USceneCaptureComponent2D
+{
+	GENERATED_UCLASS_BODY()
+
+public:
+
+#if WITH_EDITOR
+	UE_API virtual void PostEditChangeProperty(struct FPropertyChangedEvent& PropertyChangedEvent) override;
+#endif //WITH_EDITOR
+
+	/** Find or create the active scene capture render target. */
+	UE_API UTextureRenderTarget2D* FindOrCreateRenderTarget(int32 InSizeX, int32 InSizeY);
+};
+
+#undef UE_API
+
diff --git a/Engine/Plugins/Experimental/Compositing/Composite/Source/Composite/Public/Components/CompositeViewProjectionComponent.h b/Engine/Plugins/Experimental/Compositing/Composite/Source/Composite/Public/Components/CompositeViewProjectionComponent.h
new file mode 100644
--- /dev/null
+++ b/Engine/Plugins/Experimental/Compositing/Composite/Source/Composite/Public/Components/CompositeViewProjectionComponent.h
@@ -0,0 +1,89 @@
+// Copyright Epic Games, Inc. All Rights Reserved.
+
+#pragma once
+
+#include "CoreTypes.h"
+#include "Components/SceneComponent.h"
+
+#include "CompositeViewProjectionComponent.generated.h"
+
+#define UE_API COMPOSITE_API
+
+class UCameraComponent;
+class UMaterialParameterCollection;
+class UWorld;
+
+/** Component responsible for continuously updating the specified material parameter collection with a camera view projection matrix (to be used for texture projection in materials). */
+UCLASS(MinimalAPI, EditInlineNew, HideCategories = (Activation, Transform, Lighting, Rendering, Tags, Cooking, Physics, LOD, AssetUserData, Navigation), ClassGroup = Composite, editinlinenew, meta = (BlueprintSpawnableComponent, DisplayName = "Camera View Projection"))
+class UCompositeViewProjectionComponent : public UActorComponent
+{
+	GENERATED_UCLASS_BODY()
+
+public:
+	//~ Begin UObject interface
+	UE_API virtual void PostInitProperties() override;
+	UE_API virtual void PostDuplicate(bool bDuplicateForPIE) override;
+	UE_API virtual void PostEditImport() override;
+	UE_API virtual void PostLoad() override;
+	//~ End UObject interface
+
+	//~ Begin UActorComponent interface
+	UE_API virtual void OnRegister() override;
+	UE_API virtual void OnUnregister() override;
+	UE_API virtual void TickComponent(float DeltaTime, enum ELevelTick TickType, FActorComponentTickFunction* ThisTickFunction) override;
+	//~ End UActorComponent interface
+
+#if WITH_EDITOR
+	UE_API virtual void PostEditChangeProperty(struct FPropertyChangedEvent& PropertyChangedEvent) override;
+#endif //WITH_EDITOR
+
+public:
+	/** Whether or not the component activates the view-projection matrix Material Parameter Collection update. */
+	UPROPERTY(EditAnywhere, BlueprintReadWrite, Category = "CameraProjection")
+	bool bIsEnabled = true;
+
+	/** The Material Parameter Collection in which the view-projection matrix should be stored. */
+	UPROPERTY(EditAnywhere, BlueprintReadWrite, Category = "CameraProjection")
+	TObjectPtr<UMaterialParameterCollection> MaterialParameterCollection;
+
+	/** Parameter name of the first element of the transform in the Material Parameter Collection set above.  Requires space for 4 vectors. */
+	UPROPERTY(EditAnywhere, BlueprintReadWrite, Category = "CameraProjection")
+	FName MatrixParameterName;
+
+public:
+	/** Force update the Material Parameter Collection this frame. */
+	UFUNCTION(BlueprintCallable, Category = "CameraProjection")
+	UE_API void ForceUpdate();
+
+	/** Target component getter. */
+	UFUNCTION(BlueprintGetter)
+	UE_API FComponentReference& GetTargetComponent();
+
+	/** Target component setter. */
+	UFUNCTION(BlueprintSetter)
+	UE_API void SetTargetComponent(const FComponentReference& InComponentReference);
+
+private:
+	/** The CineCameraComponent whose view-projection matrix is used. */
+	UPROPERTY(EditInstanceOnly, BlueprintGetter = GetTargetComponent, BlueprintSetter = SetTargetComponent, Category = "CameraProjection", meta = (AllowPrivateAccess, UseComponentPicker, AllowedClasses = "/Script/Engine.CameraComponent, /Script/CinematicCamera.CineCameraComponent"))
+	FComponentReference TargetCameraComponent;
+
+	/** If TargetCameraComponent is not set, initialize it to the first CineCameraComponent on the same actor as this component */
+	void InitDefaultCamera();
+	
+	/** Update the Material Parameter Collection with the camera view transform information. */
+	void UpdateProjection(UWorld* InWorld) const;
+
+	/** Cached most recent target camera, used to clean up the old camera when the user changes the target */
+	UPROPERTY()
+	TObjectPtr<UCameraComponent> LastCameraComponent = nullptr;
+
+	/** Cached view projection matrix, to avoid needless MPC updates. */
+	mutable FMatrix LastViewProjectionMatrix = FMatrix::Identity;
+
+	/** Post-viewport update callback. */
+	FDelegateHandle OnWorldPreSendAllEndOfFrameUpdatesHandle;
+};
+
+#undef UE_API
+
diff --git a/Engine/Plugins/Experimental/Compositing/Composite/Source/Composite/Public/CompositeActor.h b/Engine/Plugins/Experimental/Compositing/Composite/Source/Composite/Public/CompositeActor.h
new file mode 100644
--- /dev/null
+++ b/Engine/Plugins/Experimental/Compositing/Composite/Source/Composite/Public/CompositeActor.h
@@ -0,0 +1,144 @@
+// Copyright Epic Games, Inc. All Rights Reserved.
+
+#pragma once
+
+#include "CoreTypes.h"
+#include "GameFramework/Actor.h"
+
+#include "Passes/CompositePassBase.h"
+#include "Layers/CompositeLayerBase.h"
+
+#include "CompositeActor.generated.h"
+
+#define UE_API COMPOSITE_API
+
+class ACameraActor;
+class UCompositeCoreSubsystem;
+class UCompositeShadowReflectionCatcherComponent;
+class UPrimitiveComponent;
+
+using TShadowReflectionCatcherPair = TStaticArray<TObjectPtr<UCompositeShadowReflectionCatcherComponent>, 2>;
+
+DECLARE_DELEGATE_TwoParams(FOnCompositeActorEnabledStateChange, bool /* */, UCompositeCoreSubsystem* /* */);
+
+/** Actor used to control properties of the composite pipeline. */
+UCLASS(MinimalAPI)
+class ACompositeActor : public AActor
+{
+	GENERATED_BODY()
+
+public:
+	/** Constructor. */
+	UE_API ACompositeActor(const FObjectInitializer& ObjectInitializer);
+
+	/** Destructor. */
+	UE_API ~ACompositeActor();
+
+	//~ Begin AActor Interface
+	UE_API virtual void PostRegisterAllComponents() override;
+	UE_API virtual void UnregisterAllComponents(bool bForReregister = false) override;
+	UE_API virtual void PostUnregisterAllComponents() override;
+	UE_API virtual void Tick(float DeltaSeconds) override;
+	//~ End AActor Interface
+
+#if WITH_EDITOR
+	UE_API virtual void PreEditChange(FProperty* PropertyAboutToChange) override;
+	UE_API virtual void PostEditChangeChainProperty(struct FPropertyChangedChainEvent& PropertyChangedChainEvent) override;
+	UE_API virtual void SetIsTemporarilyHiddenInEditor(bool bIsHidden) override;
+	UE_API virtual bool ShouldTickIfViewportsOnly() const override;
+#endif // WITH_EDITOR
+
+	/** Get the enabled state. */
+	UFUNCTION(BlueprintGetter)
+	UE_API bool IsEnabled() const;
+
+	/** Set the enabled state. */
+	UFUNCTION(BlueprintSetter)
+	UE_API void SetEnabled(bool bInEnabled);
+
+	/** Camera component getter. */
+	UFUNCTION(BlueprintGetter)
+	UE_API FComponentReference& GetCamera();
+
+	/** Camera component setter. */
+	UFUNCTION(BlueprintSetter)
+	UE_API void SetCamera(const FComponentReference& InComponentReference);
+
+	/** Bind to this to get a callback when the enabled state changes. */
+	FOnCompositeActorEnabledStateChange OnEnabledStateChange;
+
+private:
+	/** Component responsible for continuously updating a material parameter collection with the composite camera view projection matrix. */
+	UPROPERTY(Instanced)
+	TObjectPtr<class UCompositeViewProjectionComponent> ViewProjectionComponent;
+
+	/** Whether or not the composite behavior is active. */
+	UPROPERTY(EditAnywhere, BlueprintGetter = IsEnabled, BlueprintSetter = SetEnabled, Category = "Composite", meta = (AllowPrivateAccess = true, DisplayPriority="0"))
+	bool bIsEnabled;
+
+public:
+	/** Composite render (output) resolution. Currently only used for transient scene capture render targets. */
+	UPROPERTY(EditAnywhere, BlueprintReadWrite, Category = "Composite")
+	FIntPoint RenderResolution;
+
+private:
+	/** The primary camera used for the composite. Used by the camera view projection component for continuous view-projection matrix updates. */
+	UPROPERTY(EditInstanceOnly, BlueprintGetter = GetCamera, BlueprintSetter = SetCamera, Category = "Composite", meta = (AllowPrivateAccess, UseComponentPicker, AllowAnyActor, AllowedClasses = "/Script/Engine.CameraComponent, /Script/CinematicCamera.CineCameraComponent"))
+	FComponentReference Camera;
+
+public:
+	/** Convenience pre-configured layer of passes. */
+	UPROPERTY(EditAnywhere, BlueprintReadWrite, Instanced, Category = "Composite", meta = (DisplayName = "Layers"))
+	TArray<TObjectPtr<UCompositeLayerBase>> CompositeLayers;
+
+	/** Also provide the post-processing composite graph to SSR input. */
+	UPROPERTY(EditAnywhere, BlueprintReadWrite, Category = "Composite", AdvancedDisplay)
+	bool bEnableScreenSpaceReflections;
+
+	/** Override default view user flags. */
+	UPROPERTY(EditAnywhere, BlueprintReadWrite, Category = "Composite", AdvancedDisplay)
+	bool bOverridesViewUserFlags;
+
+	/** Custom user flags value used to alter materials in the composite render pass. Set to one by default such that branching can be used in Lit materials. */
+	UPROPERTY(EditAnywhere, BlueprintReadWrite, Category = "Composite", AdvancedDisplay, meta = (EditCondition = "bOverridesViewUserFlags"))
+	int32 ViewUserFlags;
+
+public:
+	/** Convenience function used by layers to request a reflection & shadow catcher component pair. */
+	UE_API TShadowReflectionCatcherPair FindOrCreateReflectionShadowCatcherPair(const UCompositeLayerBase* InLayer, const TArray<TObjectPtr<AActor>>& InActorsToHide);
+
+	/** Convenience function used by layers to remove and destroy a reflection & shadow catcher component pair. */
+	UE_API void DestroyReflectionShadowCatcherPair(const UCompositeLayerBase* InLayer);
+
+	/** Convenience function used by layers to control the capture state of their reflection & shadow catcher component pair. */
+	UE_API void SetReflectionShadowCatcherCaptureEveryFrame(const UCompositeLayerBase* InLayer, bool bState);
+
+	/** Returns true if the composite is active. */
+	UE_API bool IsActive() const;
+
+private:
+	/** Apply all actor state to the composite subsystem. */
+	void ApplyAllState();
+
+	/** Remove all actor state from the composite subsystem. */
+	void RemoveAllState();
+
+	/** Managed reflection & shadow catcher components. */
+	UPROPERTY()
+	TMap<TWeakObjectPtr<const UCompositeLayerBase>, TObjectPtr<UCompositeShadowReflectionCatcherComponent>> ReflectionShadowCatchers_CG;
+
+	/** Managed reflection & shadow catcher components. */
+	UPROPERTY()
+	TMap<TWeakObjectPtr<const UCompositeLayerBase>, TObjectPtr<UCompositeShadowReflectionCatcherComponent>> ReflectionShadowCatchers_NoCG;
+
+#if WITH_EDITORONLY_DATA
+	/** Used to update components after we remove layers. */
+	TArray<UCompositeLayerBase*> PreEditCompositeLayers;
+#endif
+
+	/** Flag used to determine in which cases we should ignore PostRegister / Unregister calls. */
+	bool bIsModifyingAProperty;
+};
+
+#undef UE_API
+
diff --git a/Engine/Plugins/Experimental/Compositing/Composite/Source/Composite/Public/CompositeAssetUserData.h b/Engine/Plugins/Experimental/Compositing/Composite/Source/Composite/Public/CompositeAssetUserData.h
new file mode 100644
--- /dev/null
+++ b/Engine/Plugins/Experimental/Compositing/Composite/Source/Composite/Public/CompositeAssetUserData.h
@@ -0,0 +1,27 @@
+// Copyright Epic Games, Inc. All Rights Reserved.
+
+#pragma once
+
+#include "Engine/AssetUserData.h"
+#include "CompositeAssetUserData.generated.h"
+
+class UPrimitiveComponent;
+
+DECLARE_DELEGATE_OneParam(FOnPostEditChangeOwner, UPrimitiveComponent& /* */);
+
+/**
+ * AssetUserData to keep track of edits on composite meshes.
+ */
+UCLASS()
+class UCompositeAssetUserData : public UAssetUserData
+{
+	GENERATED_BODY()
+
+public:
+	/** Bind to this to get a callback when something is changed on the associated component. */
+	FOnPostEditChangeOwner OnPostEditChangeOwner;
+	
+	//~ Begin UAssetUserData
+	virtual void PostEditChangeOwner(const FPropertyChangedEvent& PropertyChangedEvent) override;
+	//~ End UAssetUserData
+};
diff --git a/Engine/Plugins/Experimental/Compositing/Composite/Source/Composite/Public/CompositeMeshActor.h b/Engine/Plugins/Experimental/Compositing/Composite/Source/Composite/Public/CompositeMeshActor.h
new file mode 100644
--- /dev/null
+++ b/Engine/Plugins/Experimental/Compositing/Composite/Source/Composite/Public/CompositeMeshActor.h
@@ -0,0 +1,34 @@
+// Copyright Epic Games, Inc. All Rights Reserved.
+
+#pragma once
+
+#include "CoreTypes.h"
+#include "GameFramework/Actor.h"
+
+#include "CompositeMeshActor.generated.h"
+
+#define UE_API COMPOSITE_API
+
+class UCompositeMeshComponent;
+
+/** Convenience compositing mesh actor, which the composite actor can reference. */
+UCLASS(MinimalAPI)
+class ACompositeMeshActor : public AActor
+{
+	GENERATED_BODY()
+
+public:
+	UE_API ACompositeMeshActor(const FObjectInitializer& ObjectInitializer);
+	UE_API ~ACompositeMeshActor();
+
+private:
+	/** Composite mesh component. */
+	UPROPERTY(VisibleAnywhere, BlueprintReadOnly, Category = "Composite", meta = (AllowPrivateAccess = "true"))
+	TObjectPtr<UCompositeMeshComponent> CompositeMeshComponent;
+
+public:
+	/** Returns CompositeMeshComponent subobject **/
+	UE_API class UCompositeMeshComponent* GetCaptureComponent2D() const { return CompositeMeshComponent; }
+};
+
+#undef UE_API
diff --git a/Engine/Plugins/Experimental/Compositing/Composite/Source/Composite/Public/CompositeModule.h b/Engine/Plugins/Experimental/Compositing/Composite/Source/Composite/Public/CompositeModule.h
new file mode 100644
--- /dev/null
+++ b/Engine/Plugins/Experimental/Compositing/Composite/Source/Composite/Public/CompositeModule.h
@@ -0,0 +1,21 @@
+// Copyright Epic Games, Inc. All Rights Reserved.
+
+#pragma once
+
+#include "Modules/ModuleInterface.h"
+#include "Logging/LogMacros.h"
+#include "Stats/Stats.h"
+
+COMPOSITE_API DECLARE_LOG_CATEGORY_EXTERN(LogComposite, Log, All);
+
+DECLARE_STATS_GROUP(TEXT("Composite"), STATGROUP_Composite, STATCAT_Advanced);
+
+class FCompositeModule : public IModuleInterface
+{
+public:
+	// IModuleInterface interface
+	virtual void StartupModule() override;
+	virtual void ShutdownModule() override;
+	// End of IModuleInterface interface
+};
+
diff --git a/Engine/Plugins/Experimental/Compositing/Composite/Source/Composite/Public/Layers/CompositeLayerBase.h b/Engine/Plugins/Experimental/Compositing/Composite/Source/Composite/Public/Layers/CompositeLayerBase.h
new file mode 100644
--- /dev/null
+++ b/Engine/Plugins/Experimental/Compositing/Composite/Source/Composite/Public/Layers/CompositeLayerBase.h
@@ -0,0 +1,83 @@
+// Copyright Epic Games, Inc. All Rights Reserved.
+
+#pragma once
+
+#include "CoreTypes.h"
+#include "Containers/ContainersFwd.h"
+#include "UObject/Object.h"
+
+#include "Passes/CompositePassBase.h"
+#include "Passes/CompositeCorePassMergeProxy.h"
+
+#include "CompositeLayerBase.generated.h"
+
+#define UE_API COMPOSITE_API
+
+class ACompositeActor;
+
+/** Base class for defining a pass object, and its associated render proxy. */
+UCLASS(MinimalAPI, Abstract, EditInlineNew)
+class UCompositeLayerBase : public UObject
+{
+	GENERATED_BODY()
+
+public:
+	/** Context used during layer to proxy assembly. */
+	struct FTraversalContext
+	{
+		/** Flag to indicate solo pass */
+		bool bIsSolo = false;
+		
+		/** Flag to indicate first pass */
+		bool bIsFirstPass = false;
+		
+		/** List of external textures */
+		TArray<UE::CompositeCore::FExternalTexture> ExternalTextures;
+		
+		/** Array of passes applied at the start of rendering. */
+		TSortedMap<UE::CompositeCore::ResourceId, TArray<const FCompositeCorePassProxy*>> PreprocessingPasses;
+	};
+
+	/** Constructor. */
+	UE_API UCompositeLayerBase(const FObjectInitializer& ObjectInitializer);
+	
+	/** Destructor. */
+	UE_API virtual ~UCompositeLayerBase();
+
+#if WITH_EDITOR
+	UE_API virtual bool CanEditChange(const FProperty* InProperty) const;
+	UE_API virtual void PostEditChangeProperty(FPropertyChangedEvent& PropertyChangedEvent) override;
+#endif
+
+	/** Getter function to override, returning pass layer proxies to be passed to the render thread. (Proxy objects should be allocated from the provided allocator.) */
+	UE_API virtual bool GetProxy(FTraversalContext& InContext, FSceneRenderingBulkObjectAllocator& InFrameAllocator, FCompositeCorePassProxy*& OutProxy) const { return false; };
+
+public:
+	/** Whether or not the pass is enabled. */
+	UPROPERTY(EditAnywhere, BlueprintReadWrite, Category = "", meta = (DisplayPriority = "1"))
+	bool bIsEnabled = true;
+
+	/** Whether or not the pass is solo. */
+	UPROPERTY(EditAnywhere, BlueprintReadWrite, Category = "", meta = (DisplayPriority = "2"))
+	bool bIsSolo = false;
+
+#if WITH_EDITORONLY_DATA
+	/** Pass layer name. */
+	UPROPERTY()
+	FString Name;
+#endif
+
+	/** Merge operation applied on Input0 with Input1. */
+	UPROPERTY(EditAnywhere, BlueprintReadWrite, Category = "", meta = (DisplayPriority = "4"))
+	ECompositeCoreMergeOp Operation = ECompositeCoreMergeOp::Over;
+
+protected:
+	/** Convenience function that returns the merge operation depending on the traversal context. */
+	ECompositeCoreMergeOp GetOperation(const FTraversalContext& InContext) const;
+
+	/** Convenience function that returns the default secondary input depending on the traversal context. */
+	UE::CompositeCore::FPassInputDecl GetDefaultSecondInput(const FTraversalContext& InContext) const;
+};
+
+#undef UE_API
+
diff --git a/Engine/Plugins/Experimental/Compositing/Composite/Source/Composite/Public/Layers/CompositeLayerMainRender.h b/Engine/Plugins/Experimental/Compositing/Composite/Source/Composite/Public/Layers/CompositeLayerMainRender.h
new file mode 100644
--- /dev/null
+++ b/Engine/Plugins/Experimental/Compositing/Composite/Source/Composite/Public/Layers/CompositeLayerMainRender.h
@@ -0,0 +1,27 @@
+// Copyright Epic Games, Inc. All Rights Reserved.
+
+#pragma once
+
+#include "CompositeLayerBase.h"
+
+#include "CompositeLayerMainRender.generated.h"
+
+#define UE_API COMPOSITE_API
+
+UCLASS(MinimalAPI, BlueprintType, Blueprintable, EditInlineNew, CollapseCategories, meta = (DisplayName = "Main Render"))
+class UCompositeLayerMainRender : public UCompositeLayerBase
+{
+	GENERATED_BODY()
+
+public:
+	/** Constructor. */
+	UE_API UCompositeLayerMainRender(const FObjectInitializer& ObjectInitializer);
+	
+	/** Destructor. */
+	UE_API virtual ~UCompositeLayerMainRender();
+
+	/** Getter function to override, returning pass layer proxies to be passed to the render thread. (Proxy objects should be allocated from the provided allocator.) */
+	UE_API virtual bool GetProxy(FTraversalContext& InContext, FSceneRenderingBulkObjectAllocator& InFrameAllocator, FCompositeCorePassProxy*& OutProxy) const override;
+};
+
+#undef UE_API
diff --git a/Engine/Plugins/Experimental/Compositing/Composite/Source/Composite/Public/Layers/CompositeLayerPlate.h b/Engine/Plugins/Experimental/Compositing/Composite/Source/Composite/Public/Layers/CompositeLayerPlate.h
new file mode 100644
--- /dev/null
+++ b/Engine/Plugins/Experimental/Compositing/Composite/Source/Composite/Public/Layers/CompositeLayerPlate.h
@@ -0,0 +1,139 @@
+// Copyright Epic Games, Inc. All Rights Reserved.
+
+#pragma once
+
+#include "CompositeLayerBase.h"
+#include "Passes/CompositePassOpenColorIO.h"
+#include "Tickable.h"
+
+#include "CompositeLayerPlate.generated.h"
+
+#define UE_API COMPOSITE_API
+
+class UCompositeCoreSubsystem;
+class UPrimitiveComponent;
+class UTexture;
+
+UCLASS(MinimalAPI, BlueprintType, Blueprintable, EditInlineNew, CollapseCategories, meta = (DisplayName = "Plate"))
+class UCompositeLayerPlate : public UCompositeLayerBase, public FTickableGameObject
+{
+	GENERATED_BODY()
+
+public:
+	/** Composite (media) texture material parameter name. */
+	UE_API static const FLazyName CompositeTextureName;
+
+	/** Constructor. */
+	UE_API UCompositeLayerPlate(const FObjectInitializer& ObjectInitializer);
+	
+	/** Destructor. */
+	UE_API virtual ~UCompositeLayerPlate();
+
+	//~ Begin FTickableGameObject interface
+	UE_API virtual void Tick(float DeltaTime) override;
+	
+	UE_API virtual bool IsTickableWhenPaused() const override { return true; }
+	UE_API virtual bool IsTickableInEditor() const override { return true; }
+	UE_API virtual ETickableTickType GetTickableTickType() const override;
+	UE_API virtual bool IsTickable() const override;
+
+	UE_API virtual UWorld* GetTickableGameObjectWorld() const override;
+	UE_API virtual TStatId GetStatId() const override;
+	//~ End FTickableGameObject interface
+
+	/** Getter function to override, returning pass layer proxies to be passed to the render thread. (Proxy objects should be allocated from the provided allocator.) */
+	UE_API virtual bool GetProxy(FTraversalContext& InContext, FSceneRenderingBulkObjectAllocator& InFrameAllocator, FCompositeCorePassProxy*& OutProxy) const override;
+
+	/** Get array of all composite mesh primitives. */
+	UE_API TArray<UPrimitiveComponent*> GetPrimitives() const;
+
+protected:
+	//~ Begin UObject Interface
+	UE_API void PostInitProperties() override;
+	UE_API void PostLoad() override;
+
+#if WITH_EDITOR
+	UE_API virtual bool CanEditChange(const FProperty* InProperty) const;
+	UE_API virtual void PreEditChange(FProperty* PropertyThatWillChange) override;
+	UE_API virtual void PostEditChangeProperty(FPropertyChangedEvent& PropertyChangedEvent) override;
+	UE_API virtual void PostTransacted(const FTransactionObjectEvent& InTransactionEvent) override;
+#endif
+	//~ End UObject Interface
+
+public:
+
+	/** Get the enabled state. */
+	UFUNCTION(BlueprintGetter)
+	UE_API const TArray<AActor*> GetCompositeMeshes() const;
+
+	/** Set the enabled state. */
+	UFUNCTION(BlueprintSetter)
+	UE_API void SetCompositeMeshes(TArray<AActor*> InCompositeMeshes);
+
+	/** Get whether the holdout composite pipeline is active. */
+	UFUNCTION(BlueprintGetter)
+	UE_API bool IsHoldoutComposite() const { return bHoldoutComposite; }
+
+	/** Enable/disable the holdout composite pipeline. */
+	UFUNCTION(BlueprintSetter)
+	UE_API void SetHoldoutComposite(bool bInHoldoutComposite);
+
+	/** Returns the media texture, optionally pre-processed into a render target. */
+	UFUNCTION(BlueprintCallable, Category = "Composite")
+	UTexture* GetCompositeTexture() const;
+
+private:
+	/** Actors whose primitive mesh components will be marked as primitive alpha holdout, and rendered separately with the projected (media) texture. */
+	UPROPERTY(EditAnywhere, BlueprintGetter = GetCompositeMeshes, BlueprintSetter = SetCompositeMeshes, Category = "", meta = (AllowPrivateAccess = true))
+	TArray<TObjectPtr<AActor>> CompositeMeshes;
+
+public:
+	/** Media texture for compositing layer. */
+	UPROPERTY(EditAnywhere, BlueprintReadWrite, Category = "", meta = (AllowedClasses="/Script/Engine.Texture2D, /Script/MediaAssets.MediaTexture"))
+	TObjectPtr<UTexture> Texture;
+
+	/** Texture pre-processing passes, applied at the start of rendering and updating a render target. */
+	UPROPERTY(EditAnywhere, BlueprintReadWrite, Instanced, Category = "", meta = (DisplayName = "Pre-Processing", DisallowedClasses = "/Script/Composite.CompositePassMaterial"))
+	TArray<TObjectPtr<UCompositePassBase>> Preprocessing;
+
+	/** Sub-passes applied during post-processing on the texture or holdout composite render. */
+	UPROPERTY(EditAnywhere, BlueprintReadWrite, Instanced, Category = "", meta = (DisplayName = "Post-Processing"))
+	TArray<TObjectPtr<UCompositePassBase>> Postprocessing;
+
+private:
+	/** Enable holdout composite pipeline on the composite meshes. (Automatically disabled when the layer is moved above the main render.) */
+	UPROPERTY(EditAnywhere, BlueprintGetter = IsHoldoutComposite, BlueprintSetter = SetHoldoutComposite, Category = "")
+	bool bHoldoutComposite;
+
+private:
+	/** Update the (media) texture on all composite meshes. We also add a user asset data to track changes. */
+	void UpdateCompositeMeshes(bool bRegisterOnly = false) const;
+	
+	/** Update the (media) texture on a primitive component. */
+	void UpdatePrimitiveComponent(UPrimitiveComponent& InPrimitiveComponent) const;
+
+	/** Handle state change of the owner composite actor. */
+	void HandleHoldoutCompositeStateChange(bool bInEnabled, UCompositeCoreSubsystem* InSubsystem) const;
+
+	/** Convenience function verify if the composite mesh actor is already in use by another layer. */
+	bool IsCompositeMeshActorAlreadyInUse(AActor* InCompositeMeshActor) const;
+
+	/** Get number of valid primitives. */
+	int32 GetValidPrimitivesNum() const;
+
+	/** Get number of valid pre-processing passes. */
+	int32 GetValidPreprocessingPassNum() const;
+
+private:
+#if WITH_EDITOR
+	/** Pre-edit list of registered composite meshes. */
+	TArray<TObjectPtr<AActor>> PreEditCompositeMeshes;
+#endif
+
+	/** Processed texture render target. Note: Should be replaced by pool-managed RT. */
+	UPROPERTY(Transient)
+	TObjectPtr<class UTextureRenderTarget2D> TextureProcessed;
+};
+
+#undef UE_API
+
diff --git a/Engine/Plugins/Experimental/Compositing/Composite/Source/Composite/Public/Layers/CompositeLayerShadowReflection.h b/Engine/Plugins/Experimental/Compositing/Composite/Source/Composite/Public/Layers/CompositeLayerShadowReflection.h
new file mode 100644
--- /dev/null
+++ b/Engine/Plugins/Experimental/Compositing/Composite/Source/Composite/Public/Layers/CompositeLayerShadowReflection.h
@@ -0,0 +1,40 @@
+// Copyright Epic Games, Inc. All Rights Reserved.
+
+#pragma once
+
+#include "CompositeLayerBase.h"
+
+#include "Engine/TextureRenderTarget2D.h"
+
+#include "CompositeLayerShadowReflection.generated.h"
+
+#define UE_API COMPOSITE_API
+
+class AActor;
+
+UCLASS(MinimalAPI, BlueprintType, Blueprintable, EditInlineNew, CollapseCategories, meta = (DisplayName = "Shadow Reflection Catcher"))
+class UCompositeLayerShadowReflection : public UCompositeLayerBase
+{
+	GENERATED_BODY()
+
+public:
+	/** Constructor. */
+	UE_API UCompositeLayerShadowReflection(const FObjectInitializer& ObjectInitializer);
+	
+	/** Destructor. */
+	UE_API virtual ~UCompositeLayerShadowReflection();
+
+#if WITH_EDITOR
+	UE_API virtual void PostEditChangeProperty(FPropertyChangedEvent& PropertyChangedEvent) override;
+#endif
+
+	/** Getter function to override, returning pass layer proxies to be passed to the render thread. (Proxy objects should be allocated from the provided allocator.) */
+	UE_API virtual bool GetProxy(FTraversalContext& InContext, FSceneRenderingBulkObjectAllocator& InFrameAllocator, FCompositeCorePassProxy*& OutProxy) const override;
+
+	/** Also provide the post-processing composite graph to SSR input. */
+	UPROPERTY(EditAnywhere, BlueprintReadWrite, Category = "")
+	TArray<TObjectPtr<AActor>> Actors;
+};
+
+#undef UE_API
+
diff --git a/Engine/Plugins/Experimental/Compositing/Composite/Source/Composite/Public/Passes/CompositePassBase.h b/Engine/Plugins/Experimental/Compositing/Composite/Source/Composite/Public/Passes/CompositePassBase.h
new file mode 100644
--- /dev/null
+++ b/Engine/Plugins/Experimental/Compositing/Composite/Source/Composite/Public/Passes/CompositePassBase.h
@@ -0,0 +1,44 @@
+// Copyright Epic Games, Inc. All Rights Reserved.
+
+#pragma once
+
+#include "CoreTypes.h"
+#include "Containers/ContainersFwd.h"
+#include "UObject/Object.h"
+
+#include "Passes/CompositeCorePassProxy.h"
+
+#include "CompositePassBase.generated.h"
+
+#define UE_API COMPOSITE_API
+
+/** Base class for defining a pass object, and its associated render proxy. */
+UCLASS(MinimalAPI, Abstract, EditInlineNew)
+class UCompositePassBase : public UObject
+{
+	GENERATED_BODY()
+
+public:
+	/** Constructor. */
+	UE_API UCompositePassBase(const FObjectInitializer& ObjectInitializer);
+	
+	/** Destructor. */
+	UE_API virtual ~UCompositePassBase();
+
+#if WITH_EDITOR
+	UE_API virtual bool CanEditChange(const FProperty* InProperty) const;
+#endif
+
+	/** Getter function to override, returning pass proxies to be passed to the render thread. (Proxy objects should be allocated from the provided allocator.) */
+	UE_API virtual bool GetProxy(const UE::CompositeCore::FPassInputDecl& InputDecl, FSceneRenderingBulkObjectAllocator& InFrameAllocator, FCompositeCorePassProxy*& OutProxy) const { return false; };
+
+	/** Returns true if the pass is returns a valid proxy. */
+	UE_API virtual bool IsActive() const { return bIsEnabled; }
+
+	/** Whether or not the pass is active. */
+	UPROPERTY(EditAnywhere, BlueprintReadWrite, Category = "", meta = (DisplayPriority = "1"))
+	bool bIsEnabled = true;
+};
+
+#undef UE_API
+
diff --git a/Engine/Plugins/Experimental/Compositing/Composite/Source/Composite/Public/Passes/CompositePassCenteredScale.h b/Engine/Plugins/Experimental/Compositing/Composite/Source/Composite/Public/Passes/CompositePassCenteredScale.h
new file mode 100644
--- /dev/null
+++ b/Engine/Plugins/Experimental/Compositing/Composite/Source/Composite/Public/Passes/CompositePassCenteredScale.h
@@ -0,0 +1,90 @@
+// Copyright Epic Games, Inc. All Rights Reserved.
+
+#pragma once
+
+#include "CompositePassBase.h"
+
+#include "CompositePassCenteredScale.generated.h"
+
+#define UE_API COMPOSITE_API
+
+/** Primary centered scale calculation mode.*/
+UENUM()
+enum class ECompositePassScaleMode : uint8
+{
+	/** Disabled */
+	None,
+
+	/** Automatically derive aspect ratio from the parent layer media texture & composite actor camera. */
+	Automatic,
+
+	/** Calculate the scaling factor from source & target aspect ratios.*/
+	AspectRatio,
+
+	/** Manually define the scaling factor.*/
+	Manual,
+};
+
+/** Overscan uncrop derivation mode. */
+UENUM()
+enum class ECompositePassOverscanUncropMode : uint8
+{
+	/** Disabled */
+	None,
+
+	/** Automatically derive the overscan crop factor from the parent composite actor camera reference. */
+	Automatic,
+
+	/** Manually define the overscan crop factor.*/
+	Manual,
+};
+
+/** Convenience pass to re-scale footage (with black bars) inside its container texture. */
+UCLASS(MinimalAPI, BlueprintType, Blueprintable, EditInlineNew, CollapseCategories, meta=(DisplayName="Centered Scale Pass"))
+class UCompositePassCenteredScale : public UCompositePassBase
+{
+	GENERATED_BODY()
+
+public:
+	/** Constructor */
+	UE_API UCompositePassCenteredScale(const FObjectInitializer& ObjectInitializer);
+	
+	/** Destructor */
+	UE_API ~UCompositePassCenteredScale();
+
+	UE_API virtual bool GetProxy(const UE::CompositeCore::FPassInputDecl& InputDecl, FSceneRenderingBulkObjectAllocator& InFrameAllocator, FCompositeCorePassProxy*& OutProxy) const override;
+
+	UE_API virtual bool IsActive() const override;
+
+	/** Calculate the final texture UV scale. */
+	UFUNCTION(BlueprintCallable, Category = "Composite")
+	UE_API FVector2f CalculateScale() const;
+
+public:
+	/** Centered scale calculation mode used to rescale a media texture with black bars into an already constrained aspect ratio viewport. */
+	UPROPERTY(EditAnywhere, BlueprintReadWrite, Category = "")
+	ECompositePassScaleMode ScaleMode;
+
+	/** Source container aspect ratio (or resolution). */
+	UPROPERTY(EditAnywhere, BlueprintReadWrite, Category = "", meta = (EditCondition = "ScaleMode == ECompositePassScaleMode::AspectRatio"))
+	FVector2f SourceAspectRatio;
+
+	/** Embedded target aspect ratio (or resolution), without black bars. */
+	UPROPERTY(EditAnywhere, BlueprintReadWrite, Category = "", meta = (EditCondition = "ScaleMode == ECompositePassScaleMode::AspectRatio"))
+	FVector2f TargetAspectRatio;
+
+	/** Manual scale factor. */
+	UPROPERTY(EditAnywhere, BlueprintReadWrite, Category = "", meta = (EditCondition = "ScaleMode == ECompositePassScaleMode::Manual"))
+	FVector2f ScaleFactor;
+
+	/** Uncrop calculation mode used to uncrop a viewport cropped in by overscan. */
+	UPROPERTY(EditAnywhere, BlueprintReadWrite, Category = "")
+	ECompositePassOverscanUncropMode OverscanUncropMode;
+
+	/** Manual overscan to uncrop, with values from 0.0 to 1.0 matching the source camera overscan. */
+	UPROPERTY(EditAnywhere, BlueprintReadWrite, Category = "", meta = (EditCondition = "OverscanUncropMode == ECompositePassOverscanUncropMode::Manual"))
+	float Overscan;
+};
+
+#undef UE_API
+
diff --git a/Engine/Plugins/Experimental/Compositing/Composite/Source/Composite/Public/Passes/CompositePassFXAA.h b/Engine/Plugins/Experimental/Compositing/Composite/Source/Composite/Public/Passes/CompositePassFXAA.h
new file mode 100644
--- /dev/null
+++ b/Engine/Plugins/Experimental/Compositing/Composite/Source/Composite/Public/Passes/CompositePassFXAA.h
@@ -0,0 +1,34 @@
+// Copyright Epic Games, Inc. All Rights Reserved.
+
+#pragma once
+
+#include "CompositePassBase.h"
+
+#include "CompositePassFXAA.generated.h"
+
+#define UE_API COMPOSITE_API
+
+/** FXAA holdout composite pass, mostly useful for applying on the composite render pass for CG/Motion Graphics. */
+UCLASS(MinimalAPI, BlueprintType, Blueprintable, EditInlineNew, CollapseCategories, meta=(DisplayName="FXAA Pass"))
+class UCompositePassFXAA : public UCompositePassBase
+{
+	GENERATED_BODY()
+
+public:
+	/** Constructor */
+	UE_API UCompositePassFXAA(const FObjectInitializer& ObjectInitializer);
+	
+	/** Destructor */
+	UE_API ~UCompositePassFXAA();
+
+	UE_API virtual bool GetProxy(const UE::CompositeCore::FPassInputDecl& InputDecl, FSceneRenderingBulkObjectAllocator& InFrameAllocator, FCompositeCorePassProxy*& OutProxy) const override;
+
+public:
+
+	/** FXAA quality control, matching the "r.fxaa.quality" console variable values. */
+	UPROPERTY(EditAnywhere, BlueprintReadWrite, Category = "", meta=(ClampMin="0", ClampMax="5", UIMin = "0", UIMax = "5"))
+	int32 Quality;
+};
+
+#undef UE_API
+
diff --git a/Engine/Plugins/Experimental/Compositing/Composite/Source/Composite/Public/Passes/CompositePassMaterial.h b/Engine/Plugins/Experimental/Compositing/Composite/Source/Composite/Public/Passes/CompositePassMaterial.h
new file mode 100644
--- /dev/null
+++ b/Engine/Plugins/Experimental/Compositing/Composite/Source/Composite/Public/Passes/CompositePassMaterial.h
@@ -0,0 +1,34 @@
+// Copyright Epic Games, Inc. All Rights Reserved.
+
+#pragma once
+
+#include "CompositePassBase.h"
+
+#include "CompositePassMaterial.generated.h"
+
+#define UE_API COMPOSITE_API
+
+class UMaterialInterface;
+
+/** Post-process material holdout composite pass. */
+UCLASS(MinimalAPI, BlueprintType, Blueprintable, EditInlineNew, CollapseCategories, meta = (DisplayName = "Post-Process Material Pass"))
+class UCompositePassMaterial : public UCompositePassBase
+{
+	GENERATED_BODY()
+
+public:
+	/** Constructor */
+	UE_API UCompositePassMaterial(const FObjectInitializer& ObjectInitializer);
+
+	/** Destructor */
+	UE_API ~UCompositePassMaterial();
+
+	UE_API virtual bool GetProxy(const UE::CompositeCore::FPassInputDecl& InputDecl, FSceneRenderingBulkObjectAllocator& InFrameAllocator, FCompositeCorePassProxy*& OutProxy) const override;
+
+public:
+	/** Post-process material to execute. Input0 is connected to SceneTexture's PostprocessInput0. */
+	UPROPERTY(EditAnywhere, BlueprintReadWrite, Category = "")
+	TObjectPtr<UMaterialInterface> PostProcessMaterial;
+};
+
+#undef UE_API
diff --git a/Engine/Plugins/Experimental/Compositing/Composite/Source/Composite/Public/Passes/CompositePassOpenColorIO.h b/Engine/Plugins/Experimental/Compositing/Composite/Source/Composite/Public/Passes/CompositePassOpenColorIO.h
new file mode 100644
--- /dev/null
+++ b/Engine/Plugins/Experimental/Compositing/Composite/Source/Composite/Public/Passes/CompositePassOpenColorIO.h
@@ -0,0 +1,35 @@
+// Copyright Epic Games, Inc. All Rights Reserved.
+
+#pragma once
+
+#include "CoreTypes.h"
+#include "CompositePassBase.h"
+#include "Passes/CompositeCorePassMergeProxy.h"
+
+#include "OpenColorIOColorSpace.h"
+
+#include "CompositePassOpenColorIO.generated.h"
+
+#define UE_API COMPOSITE_API
+
+/** Pass to apply an OpenColorIO transform. */
+UCLASS(MinimalAPI, BlueprintType, Blueprintable, EditInlineNew, CollapseCategories, meta = (DisplayName = "OpenColorIO Pass"))
+class UCompositePassOpenColorIO : public UCompositePassBase
+{
+	GENERATED_BODY()
+
+public:
+	/** Constructor */
+	UE_API UCompositePassOpenColorIO(const FObjectInitializer& ObjectInitializer);
+	/** Destructor */
+	UE_API ~UCompositePassOpenColorIO();
+
+	UE_API virtual bool GetProxy(const UE::CompositeCore::FPassInputDecl& InputDecl, FSceneRenderingBulkObjectAllocator& InFrameAllocator, FCompositeCorePassProxy*& OutProxy) const override;
+
+public:
+	/** Color conversion settings. */
+	UPROPERTY(EditAnywhere, BlueprintReadWrite, Category = "")
+	FOpenColorIOColorConversionSettings ColorConversionSettings;
+};
+
+#undef UE_API
diff --git a/Engine/Plugins/Experimental/Compositing/Composite/Source/CompositeEditor/CompositeEditor.Build.cs b/Engine/Plugins/Experimental/Compositing/Composite/Source/CompositeEditor/CompositeEditor.Build.cs
new file mode 100644
--- /dev/null
+++ b/Engine/Plugins/Experimental/Compositing/Composite/Source/CompositeEditor/CompositeEditor.Build.cs
@@ -0,0 +1,29 @@
+// Copyright Epic Games, Inc. All Rights Reserved.
+
+using UnrealBuildTool;
+
+public class CompositeEditor : ModuleRules
+{
+	public CompositeEditor(ReadOnlyTargetRules Target) : base(Target)
+	{
+		PCHUsage = ModuleRules.PCHUsageMode.UseExplicitOrSharedPCHs;
+		
+		PublicDependencyModuleNames.AddRange(
+			new string[]
+			{
+				"Core",
+				"Composite"
+			}
+		);
+		
+		PrivateDependencyModuleNames.AddRange(
+			new string[]
+			{
+				"CoreUObject",
+				"Engine",
+				"UnrealEd",
+                "Slate",
+			}
+		);
+	}
+}
diff --git a/Engine/Plugins/Experimental/Compositing/Composite/Source/CompositeEditor/Private/CompositeEditorModule.cpp b/Engine/Plugins/Experimental/Compositing/Composite/Source/CompositeEditor/Private/CompositeEditorModule.cpp
new file mode 100644
--- /dev/null
+++ b/Engine/Plugins/Experimental/Compositing/Composite/Source/CompositeEditor/Private/CompositeEditorModule.cpp
@@ -0,0 +1,79 @@
+// Copyright Epic Games, Inc. All Rights Reserved.
+
+#include "CompositeEditorModule.h"
+
+#include "Editor.h"
+#include "EngineUtils.h"
+#include "EngineAnalytics.h"
+#include "Framework/Notifications/NotificationManager.h"
+#include "LevelEditor.h"
+#include "Modules/ModuleManager.h"
+#include "Widgets/Notifications/SNotificationList.h"
+
+#include "CompositeActor.h"
+
+#define LOCTEXT_NAMESPACE "CompositeEditorModule"
+
+DEFINE_LOG_CATEGORY(LogCompositeEditor);
+
+void FCompositeEditorModule::StartupModule()
+{
+	FCoreDelegates::OnPostEngineInit.AddRaw(this, &FCompositeEditorModule::OnPostEngineInit);
+}
+
+void FCompositeEditorModule::ShutdownModule()
+{
+	if (GEditor != nullptr)
+	{
+		GEditor->OnLevelActorAdded().RemoveAll(this);
+	}
+}
+
+void FCompositeEditorModule::OnPostEngineInit()
+{
+	if (GEditor)
+	{
+		GEditor->OnLevelActorAdded().AddRaw(this, &FCompositeEditorModule::OnLevelActorAdded);
+	}
+}
+
+void FCompositeEditorModule::OnLevelActorAdded(AActor* InActor)
+{
+	if (!InActor || InActor->HasAnyFlags(RF_Transient) || !InActor->IsA<ACompositeActor>())
+	{
+		return;
+	}
+
+	UWorld* World = InActor->GetWorld();
+	if (!IsValid(World))
+	{
+		return;
+	}
+
+	int32 CompositeActorCount = 0;
+	for (TActorIterator<ACompositeActor> It(World); It; ++It)
+	{
+		CompositeActorCount++;
+
+		if (CompositeActorCount > 1)
+		{
+			UE_LOG(LogCompositeEditor, Error, TEXT("No more than one composite actor is currently allowed."));
+
+			FNotificationInfo NotifyInfo(LOCTEXT("CompositeActorLimitMessage", "No more than one composite actor is currently allowed."));
+			NotifyInfo.ExpireDuration = 4.0f;
+			FSlateNotificationManager::Get().AddNotification(NotifyInfo)->SetCompletionState(SNotificationItem::CS_Fail);
+
+			InActor->Destroy();
+			break;
+		}
+	}
+
+	if (FEngineAnalytics::IsAvailable())
+	{
+		FEngineAnalytics::GetProvider().RecordEvent(TEXT("Editor.Usage.Composite.AddCompositeActorToLevel"));
+	}
+}
+
+IMPLEMENT_MODULE(FCompositeEditorModule, CompositeEditor)
+
+#undef LOCTEXT_NAMESPACE
diff --git a/Engine/Plugins/Experimental/Compositing/Composite/Source/CompositeEditor/Public/CompositeEditorModule.h b/Engine/Plugins/Experimental/Compositing/Composite/Source/CompositeEditor/Public/CompositeEditorModule.h
new file mode 100644
--- /dev/null
+++ b/Engine/Plugins/Experimental/Compositing/Composite/Source/CompositeEditor/Public/CompositeEditorModule.h
@@ -0,0 +1,26 @@
+// Copyright Epic Games, Inc. All Rights Reserved.
+
+#pragma once
+
+#include "CoreMinimal.h"
+#include "LevelEditor.h"
+#include "Logging/LogMacros.h"
+#include "Modules/ModuleManager.h"
+
+/** Log category for this module. */
+DECLARE_LOG_CATEGORY_EXTERN(LogCompositeEditor, Log, All);
+
+class FCompositeEditorModule : public IModuleInterface
+{
+public:
+	
+	/** IModuleInterface implementation */
+	virtual void StartupModule() override;
+	virtual void ShutdownModule() override;
+
+	/** Called after the engine has initialized. */
+	void OnPostEngineInit();
+
+	/** Called when an Actor has been added to a level. */
+	void OnLevelActorAdded(AActor* InActor);
+};
