Composite: Added support for color grade pass.

diff --git a/Engine/Plugins/Experimental/Compositing/Composite/Shaders/Private/CompositeColorGrade.usf b/Engine/Plugins/Experimental/Compositing/Composite/Shaders/Private/CompositeColorGrade.usf
new file mode 100644
--- /dev/null
+++ b/Engine/Plugins/Experimental/Compositing/Composite/Shaders/Private/CompositeColorGrade.usf
@@ -0,0 +1,117 @@
+// Copyright Epic Games, Inc. All Rights Reserved.
+
+#include "/Engine/Private/Common.ush"
+#include "/Engine/Private/TonemaCommon.ush"
+#include "/Engine/Private/ACES/ACESCommon.ush"
+
+static const float DefaultTemperature = 6500.0;
+static const float DefaultTint = 0.0;
+static const float Epsilon = 1e-8;
+
+Texture2D InputTexture;
+SamplerState InputSampler;
+
+uint bIsTemperatureWhiteBalance;
+float WhiteTemp;
+float WhiteTint;
+
+float4 Global_Saturation;
+float4 Global_Contrast;
+float4 Global_Gamma;
+float4 Global_Gain;
+float4 Global_Offset;
+
+float4 Shadows_Saturation;
+float4 Shadows_Contrast;
+float4 Shadows_Gamma;
+float4 Shadows_Gain;
+float4 Shadows_Offset;
+
+float4 Midtones_Saturation;
+float4 Midtones_Contrast;
+float4 Midtones_Gamma;
+float4 Midtones_Gain;
+float4 Midtones_Offset;
+
+float4 Highlights_Saturation;
+float4 Highlights_Contrast;
+float4 Highlights_Gamma;
+float4 Highlights_Gain;
+float4 Highlights_Offset;
+
+float ShadowsMax;
+float HighlightsMin;
+float HighlightsMax;
+
+float3 ColorCorrect( float3 WorkingColor,
+	float4 ColorSaturation,
+	float4 ColorContrast,
+	float4 ColorGamma,
+	float4 ColorGain,
+	float4 ColorOffset )
+{
+	float Luma = dot( WorkingColor, AP1_RGB2Y );
+	WorkingColor = max( 0, lerp( Luma.xxx, WorkingColor, ColorSaturation.xyz*ColorSaturation.w ) );
+	WorkingColor = pow( WorkingColor * (1.0 / 0.18), ColorContrast.xyz*ColorContrast.w ) * 0.18;
+	WorkingColor = pow( WorkingColor, 1.0 / (ColorGamma.xyz*ColorGamma.w) );
+	WorkingColor = WorkingColor * (ColorGain.xyz * ColorGain.w) + (ColorOffset.xyz + ColorOffset.w);
+	return WorkingColor;
+}
+
+// Nuke-style Color Correct
+float3 ColorCorrectAll( float3 WorkingColor )
+{
+	float Luma = dot( WorkingColor, AP1_RGB2Y );
+
+	// Shadow CC
+	float3 CCColorShadows = ColorCorrect(WorkingColor, 
+		Shadows_Saturation*Global_Saturation, 
+		Shadows_Contrast*Global_Contrast, 
+		Shadows_Gamma*Global_Gamma, 
+		Shadows_Gain*Global_Gain, 
+		Shadows_Offset+Global_Offset);
+	float CCWeightShadows = 1- smoothstep(0, ShadowsMax, Luma);
+	
+	// Highlight CC
+	float3 CCColorHighlights = ColorCorrect(WorkingColor, 
+		Highlights_Saturation*Global_Saturation, 
+		Highlights_Contrast*Global_Contrast, 
+		Highlights_Gamma*Global_Gamma, 
+		Highlights_Gain*Global_Gain, 
+		Highlights_Offset+Global_Offset);
+	float CCWeightHighlights = smoothstep(HighlightsMin, HighlightsMax, Luma);
+
+	// Midtone CC
+	float3 CCColorMidtones = ColorCorrect(WorkingColor, 
+		Midtones_Saturation*Global_Saturation, 
+		Midtones_Contrast*Global_Contrast, 
+		Midtones_Gamma*Global_Gamma, 
+		Midtones_Gain*Global_Gain, 
+		Midtones_Offset+Global_Offset);
+	float CCWeightMidtones = 1 - CCWeightShadows - CCWeightHighlights;
+
+	// Blend Shadow, Midtone and Highlight CCs
+	float3 WorkingColorSMH = CCColorShadows*CCWeightShadows + CCColorMidtones*CCWeightMidtones + CCColorHighlights*CCWeightHighlights;
+	
+	return WorkingColorSMH;
+}
+
+void MainPS(
+    noperspective float4 InUVAndScreenPos : TEXCOORD0,
+    out float4 OutColor : SV_Target0
+    )
+{
+    float4 LinearColor = InputTexture.Sample(InputSampler, InUVAndScreenPos.xy);
+    
+    BRANCH
+	if( abs(WhiteTemp - DefaultTemperature) > Epsilon || abs(WhiteTint - DefaultTint) > Epsilon )
+	{
+		LinearColor = WhiteBalance(LinearColor, WhiteTemp, WhiteTint, bIsTemperatureWhiteBalance, (float3x3)WorkingColorSpace.ToXYZ, (float3x3)WorkingColorSpace.FromXYZ);
+	}
+
+    float4 ColorAP1 = mul(WorkingColorSpace.ToAP1, LinearColor);
+    
+    ColorAP1.rgb = ColorCorrectAll(ColorAP1.rgb);
+
+    OutColor = mul(WorkingColorSpace.FromAP1, ColorAP1);
+}
diff --git a/Engine/Plugins/Experimental/Compositing/Composite/Source/Composite/Private/Passes/CompositePassColorGrade.cpp b/Engine/Plugins/Experimental/Compositing/Composite/Source/Composite/Private/Passes/CompositePassColorGrade.cpp
new file mode 100644
--- /dev/null
+++ b/Engine/Plugins/Experimental/Compositing/Composite/Source/Composite/Private/Passes/CompositePassColorGrade.cpp
@@ -0,0 +1,126 @@
+// Copyright Epic Games, Inc. All Rights Reserved.
+
+#include "Passes/CompositePassColorGrade.h"
+
+#include "SceneManagement.h"
+#include "ScreenPass.h"
+#include "PostProcess/PostProcessMaterialInputs.h"
+#include "Passes/CompositeCorePassProxy.h"
+
+BEGIN_SHADER_PARAMETER_STRUCT(FColorGradePerRangeParameters, )
+	SHADER_PARAMETER(FVector4f, Saturation)
+	SHADER_PARAMETER(FVector4f, Contrast)
+	SHADER_PARAMETER(FVector4f, Gamma)
+	SHADER_PARAMETER(FVector4f, Gain)
+	SHADER_PARAMETER(FVector4f, Offset)
+END_SHADER_PARAMETER_STRUCT()
+
+class FCompositeColorGradeShader : public FGlobalShader
+{
+	DECLARE_GLOBAL_SHADER(FCompositeColorGradeShader);
+	SHADER_USE_PARAMETER_STRUCT(FCompositeColorGradeShader, FGlobalShader);
+
+	BEGIN_SHADER_PARAMETER_STRUCT(FParameters, )
+		SHADER_PARAMETER_RDG_TEXTURE(Texture2D, InputTexture)
+		SHADER_PARAMETER_SAMPLER(SamplerState, InputSampler)
+		SHADER_PARAMETER_STRUCT_REF(FWorkingColorSpaceShaderParameters, WorkingColorSpace)
+		SHADER_PARAMETER_STRUCT(FColorGradePerRangeParameters, Global)
+		SHADER_PARAMETER_STRUCT(FColorGradePerRangeParameters, Shadows)
+		SHADER_PARAMETER_STRUCT(FColorGradePerRangeParameters, Midtones)
+		SHADER_PARAMETER_STRUCT(FColorGradePerRangeParameters, Highlights)
+		SHADER_PARAMETER(float, ShadowsMax)
+		SHADER_PARAMETER(float, HighlightsMin)
+		SHADER_PARAMETER(float, HighlightsMax)
+		SHADER_PARAMETER(uint32, bIsTemperatureWhiteBalance)
+		SHADER_PARAMETER(float, WhiteTemp)
+		SHADER_PARAMETER(float, WhiteTint)
+		RENDER_TARGET_BINDING_SLOTS()
+	END_SHADER_PARAMETER_STRUCT()
+
+	static FColorGradePerRangeParameters ToParameters(const FColorGradePerRangeSettings& InParams)
+	{
+		FColorGradePerRangeParameters Output;
+		Output.Saturation = FVector4f(InParams.Saturation);
+		Output.Contrast = FVector4f(InParams.Contrast);
+		Output.Gamma = FVector4f(InParams.Gamma);
+		Output.Gain = FVector4f(InParams.Gain);
+		Output.Offset = FVector4f(InParams.Offset);
+
+		return Output;
+	}
+};
+IMPLEMENT_GLOBAL_SHADER(FCompositeColorGradeShader, "/Plugin/Composite/Private/CompositeColorGrade.usf", "MainPS", SF_Pixel);
+
+class FColorGradeCompositePassProxy : public FCompositeCorePassProxy
+{
+public:
+	FColorGradeCompositePassProxy(UE::CompositeCore::FPassInputDeclArray InPassDeclaredInputs)
+		: FCompositeCorePassProxy(MoveTemp(InPassDeclaredInputs))
+	{ }
+
+	UE::CompositeCore::FPassOutput Add(FRDGBuilder& GraphBuilder, const FSceneView& InView, const UE::CompositeCore::FPassInputArray& Inputs, const UE::CompositeCore::FPassContext& PassContext) const override
+	{
+		const UE::CompositeCore::FResourceMetadata& Metadata = Inputs[0].Metadata;
+		const FScreenPassTexture& Input = Inputs[0].Texture;
+		FScreenPassRenderTarget Output = Inputs.OverrideOutput;
+
+		// If the override output is provided, it means that this is the last pass in post processing.
+		if (!Output.IsValid())
+		{
+			Output = CreateOutputRenderTarget(GraphBuilder, InView, PassContext.OutputViewRect, Input.Texture->Desc, TEXT("OcioCompositePass"));
+		}
+
+		FGlobalShaderMap* GlobalShaderMap = GetGlobalShaderMap(InView.GetFeatureLevel());
+		FScreenPassTextureViewport Viewport = FScreenPassTextureViewport(Output);
+
+		FCompositeColorGradeShader::FParameters* Parameters = GraphBuilder.AllocParameters<FCompositeColorGradeShader::FParameters>();
+		Parameters->InputTexture = Input.Texture;
+		Parameters->InputSampler = TStaticSamplerState<SF_Bilinear, AM_Border, AM_Border, AM_Clamp>::GetRHI();
+		Parameters->WorkingColorSpace = GDefaultWorkingColorSpaceUniformBuffer.GetUniformBufferRef();
+		Parameters->Global = FCompositeColorGradeShader::ToParameters(ColorGradingSettings.Global);
+		Parameters->Shadows = FCompositeColorGradeShader::ToParameters(ColorGradingSettings.Shadows);
+		Parameters->Midtones = FCompositeColorGradeShader::ToParameters(ColorGradingSettings.Midtones);
+		Parameters->Highlights = FCompositeColorGradeShader::ToParameters(ColorGradingSettings.Highlights);
+		Parameters->ShadowsMax = ColorGradingSettings.ShadowsMax;
+		Parameters->HighlightsMin = ColorGradingSettings.HighlightsMin;
+		Parameters->HighlightsMax = ColorGradingSettings.HighlightsMax;
+		Parameters->bIsTemperatureWhiteBalance = TemperatureSettings.TemperatureType;
+		Parameters->WhiteTemp = TemperatureSettings.WhiteTemp;
+		Parameters->WhiteTint = TemperatureSettings.WhiteTint;
+		Parameters->RenderTargets[0] = Output.GetRenderTargetBinding();
+
+		TShaderMapRef<FCompositeColorGradeShader> PixelShader(GlobalShaderMap);
+		AddDrawScreenPass(
+			GraphBuilder,
+			RDG_EVENT_NAME("Composite.CenteredScale (%dx%d)", Viewport.Extent.X, Viewport.Extent.Y),
+			InView,
+			Viewport,
+			Viewport,
+			PixelShader,
+			Parameters
+		);
+
+		return UE::CompositeCore::FPassOutput{ MoveTemp(Output), Metadata, PassOutputOverride };
+	}
+
+	FCompositeTemperatureSettings TemperatureSettings;
+	FColorGradingSettings ColorGradingSettings;
+};
+
+UCompositePassColorGrade::UCompositePassColorGrade(const FObjectInitializer& ObjectInitializer)
+	: Super(ObjectInitializer)
+{
+}
+
+UCompositePassColorGrade::~UCompositePassColorGrade() = default;
+
+bool UCompositePassColorGrade::GetProxy(const UE::CompositeCore::FPassInputDecl& InputDecl, FSceneRenderingBulkObjectAllocator& InFrameAllocator, FCompositeCorePassProxy*& OutProxy) const
+{
+	FColorGradeCompositePassProxy* Proxy = InFrameAllocator.Create<FColorGradeCompositePassProxy>(UE::CompositeCore::FPassInputDeclArray{ InputDecl });
+	Proxy->TemperatureSettings = TemperatureSettings;
+	Proxy->ColorGradingSettings = ColorGradingSettings;
+
+	OutProxy = Proxy;
+	return true;
+}
+
diff --git a/Engine/Plugins/Experimental/Compositing/Composite/Source/Composite/Public/Passes/CompositePassColorGrade.h b/Engine/Plugins/Experimental/Compositing/Composite/Source/Composite/Public/Passes/CompositePassColorGrade.h
new file mode 100644
--- /dev/null
+++ b/Engine/Plugins/Experimental/Compositing/Composite/Source/Composite/Public/Passes/CompositePassColorGrade.h
@@ -0,0 +1,61 @@
+// Copyright Epic Games, Inc. All Rights Reserved.
+
+#pragma once
+
+#include "CoreTypes.h"
+#include "CompositePassBase.h"
+#include "Passes/CompositeCorePassMergeProxy.h"
+
+#include "CompositePassColorGrade.generated.h"
+
+#define UE_API COMPOSITE_API
+
+/** Color grading pass temperature settings. */
+USTRUCT(BlueprintType)
+struct FCompositeTemperatureSettings
+{
+	GENERATED_USTRUCT_BODY()
+
+	/**
+	* Selects the type of temperature calculation.
+	* White Balance uses the Temperature value to control the virtual camera's White Balance. This is the default selection.
+	* Color Temperature uses the Temperature value to adjust the color temperature of the scene, which is the inverse of the White Balance operation.
+	*/
+	UPROPERTY(interp, BlueprintReadWrite, Category = "Color Grading|Temperature", meta = (DisplayName = "Temperature Type"))
+	TEnumAsByte<enum ETemperatureMethod> TemperatureType = ETemperatureMethod::TEMP_WhiteBalance;
+
+	/** Controls the color temperature or white balance in degrees Kelvin which the scene considers as white light. */
+	UPROPERTY(interp, BlueprintReadWrite, Category = "Color Grading|Temperature", meta = (UIMin = "1500.0", UIMax = "15000.0", DisplayName = "Temp"))
+	float WhiteTemp = 6500.0f;
+	
+	/** Controls the color of the scene along the magenta - green axis (orthogonal to the color temperature).  This feature is equivalent to color tint in digital cameras. */
+	UPROPERTY(interp, BlueprintReadWrite, Category = "Color Grading|Temperature", meta = (UIMin = "-1.0", UIMax = "1.0", DisplayName = "Tint"))
+	float WhiteTint = 0.0f;
+
+};
+
+/** Pass to apply an OpenColorIO transform. */
+UCLASS(MinimalAPI, BlueprintType, Blueprintable, EditInlineNew, CollapseCategories, meta = (DisplayName = "Color Grade Pass"))
+class UCompositePassColorGrade : public UCompositePassBase
+{
+	GENERATED_BODY()
+
+public:
+	/** Constructor */
+	UE_API UCompositePassColorGrade(const FObjectInitializer& ObjectInitializer);
+	/** Destructor */
+	UE_API ~UCompositePassColorGrade();
+
+	UE_API virtual bool GetProxy(const UE::CompositeCore::FPassInputDecl& InputDecl, FSceneRenderingBulkObjectAllocator& InFrameAllocator, FCompositeCorePassProxy*& OutProxy) const override;
+
+public:
+	/** Color temperature settings. */
+	UPROPERTY(BlueprintReadWrite, EditAnywhere, Category = "Color Grading")
+	FCompositeTemperatureSettings TemperatureSettings;
+
+	/** Color grade settings. */
+	UPROPERTY(BlueprintReadWrite, EditAnywhere, Category = "Color Grading", meta = (ShowPostProcessCategories))
+	FColorGradingSettings ColorGradingSettings;
+};
+
+#undef UE_API
