Composite: First implementation of a pre-processing keyer pass.

* Chromaticity keying is done in the (L)UV color space.
* Additional fixes: render target creation on plate layer load, support for UAV.
* (Pre-)blur and erosion improvements to come later, though impacted the choice of using compute for future LDS use.

diff --git a/Engine/Plugins/Experimental/Compositing/Composite/Shaders/Private/CompositeCommon.ush b/Engine/Plugins/Experimental/Compositing/Composite/Shaders/Private/CompositeCommon.ush
new file mode 100644
--- /dev/null
+++ b/Engine/Plugins/Experimental/Compositing/Composite/Shaders/Private/CompositeCommon.ush
@@ -0,0 +1,42 @@
+// Copyright Epic Games, Inc. All Rights Reserved.
+
+#pragma once
+
+#define ESourceEncoding_None  0
+#define ESourceEncoding_Gamma 1
+#define ESourceEncoding_sRGB  2
+
+float3 Decode(float3 InColor, uint InEncoding)
+{
+    BRANCH
+    if(InEncoding == ESourceEncoding_Gamma)
+    {
+        return pow(InColor, 2.2);
+    }
+    else if(InEncoding == ESourceEncoding_sRGB)
+    {
+        return sRGBToLinear(InColor);
+    }
+    else
+    {
+        return InColor;
+    }
+}
+
+float3 Encode(float3 InColor, uint InEncoding)
+{
+    BRANCH
+    if(InEncoding == ESourceEncoding_Gamma)
+    {
+        return pow(InColor, 2.2);
+    }
+    else if(InEncoding == ESourceEncoding_sRGB)
+    {
+        return LinearToSrgb(InColor);
+    }
+    else
+    {
+        return InColor;
+    }
+}
+
diff --git a/Engine/Plugins/Experimental/Compositing/Composite/Shaders/Private/CompositeDeltaKeyer.usf b/Engine/Plugins/Experimental/Compositing/Composite/Shaders/Private/CompositeDeltaKeyer.usf
new file mode 100644
--- /dev/null
+++ b/Engine/Plugins/Experimental/Compositing/Composite/Shaders/Private/CompositeDeltaKeyer.usf
@@ -0,0 +1,206 @@
+// Copyright Epic Games, Inc. All Rights Reserved.
+
+#include "/Engine/Private/Common.ush"
+
+#include "/Engine/Private/ACES/ACESCommon.ush"
+#include "/Engine/Private/ColorSpace.ush"
+#include "/Engine/Private/GammaCorrectionCommon.ush"
+#include "/Engine/Private/ScreenPass.ush"
+#include "/Engine/Private/PostProcessCommon.ush"
+
+#include "CompositeCommon.ush"
+
+#ifndef TILE_SIZE
+    #error TILE_SIZE must be defined
+#endif
+
+#define DESPILL_METHOD_NONE 0
+#define DESPILL_METHOD_GREEN_AVG 1
+#define DESPILL_METHOD_BLUE_AVG 2
+#define DESPILL_METHOD_ALPHA_KEY 3
+
+Texture2D InputTexture;
+SamplerState InputSampler;
+
+Texture2D CleanPlateTexture;
+SamplerState CleanPlateSampler;
+
+RWTexture2D<float4> RWOutputTexture;
+
+SCREEN_PASS_TEXTURE_VIEWPORT(Input)
+SCREEN_PASS_TEXTURE_VIEWPORT(Output)
+
+uint InputEncoding;
+uint CleanPlateEncoding;
+
+float3 KeyColor;
+float4 Params0;
+float4 Params1;
+uint bInvertAlpha;
+uint DespillMethod;
+
+// CIE 1931 D65 white point
+static const float3 XYZ_n = xyY_2_XYZ(float3(0.3127, 0.3290, 1.0));
+
+static const float DeltaE_JND = 1.0;
+static const float Lstar_CIELUV = 100.0;
+// Hand-picked scale factor used to adjust LUV color delta range
+static const float kDistanceScaleFactor_CIELUV = 5.0 * DeltaE_JND / Lstar_CIELUV;
+
+// Convert color from XYZ to CIE 1976 L*, u*, v* color space.
+// We use the latter for its improved perceptual uniformity, as a way to measure chromaticity distance for keying.
+float3 CIE_XYZ_to_LUV(float3 XYZ)
+{
+	const float Epsilon = 216.0f / 24389.0f; // (6/29)^3
+	const float Kappa = 24389.0f / 27.0f; // (29/3)^3
+
+	float L = select(XYZ[1] / XYZ_n[1] > Epsilon, 116 * pow(XYZ[1]/ XYZ_n[1], 1.0/3.0) - 16, Kappa * XYZ[1] / XYZ_n[1]);
+
+	if (L == 0.0f)
+	{
+		return float3(0.0f, 0.0f, 0.0f);
+	}
+	
+	float Denom = XYZ[0] + 15 * XYZ[1] + 3 * XYZ[2];
+	float Denom_n = XYZ_n[0] + 15 * XYZ_n[1] + 3 * XYZ_n[2];
+	
+	float up = select(Denom > 0.0f, 4 * XYZ[0] / Denom, 0.0f);
+	float vp = select(Denom > 0.0f, 9 * XYZ[1] / Denom, 0.0f);
+
+	float u_n = 4 * XYZ_n[0] / Denom_n;
+	float v_n = 9 * XYZ_n[1] / Denom_n;
+
+	float u = 13 * L * (up - u_n);
+	float v = 13 * L * (vp - v_n);
+
+	return float3(L, u, v);
+}
+
+float CalcColorDelta(float3 ColorA, float3 ColorB)
+{
+	float3 LUV_A = CIE_XYZ_to_LUV(LinearRGB_2_XYZ(ColorA));
+	float3 LUV_B = CIE_XYZ_to_LUV(LinearRGB_2_XYZ(ColorB));
+	
+	// We only calculate the distance in the (luminance-relative) chromaticity dimensions
+	float Delta = distance(LUV_A.yz, LUV_B.yz);
+
+	return kDistanceScaleFactor_CIELUV * Delta;
+}
+
+float CalculateAlpha(float Edge0, float Edge1, float Distance)
+{
+	// Colors within the specified tolerance are fully translucent, with a soft transition depending on the edge smoothness
+	float Alpha = smoothstep(Edge0, Edge1, Distance);
+
+	const float AlphaMin = Params0.z;
+	const float AlphaMax = Params0.w;
+	
+	// Apply alpha clip range (equivalent to saturate by default)
+	Alpha = select(Alpha < AlphaMin, 0.0f, Alpha);
+	Alpha = select(Alpha > AlphaMax, 1.0f, Alpha);
+
+	// Apply optional alpha inversion
+	Alpha = select(bInvertAlpha > 0, 1.0f - Alpha, Alpha);
+
+	return Alpha;
+}
+
+float3 AlphaKeyDespill(float3 Color, float AlphaMask)
+{
+	float Weight = AlphaMask * dot(Color.rgb, KeyColor) / dot(KeyColor, KeyColor);
+
+	return max(0.0f, Color.rgb - Weight * KeyColor.rgb);
+}
+
+float3 GreenDespillAvg(float3 Color)
+{
+	float Avg = 0.5f * (Color.r + Color.b);
+
+	return float3(
+		Color.r,
+		select(Color.g > Avg, Avg, Color.g),
+		Color.b
+	);
+}
+
+float3 BlueDespillAvg(float3 Color)
+{
+	float Avg = 0.5f * (Color.r + Color.g);
+
+	return float3(
+		Color.r,
+		Color.g,
+		select(Color.b > Avg, Avg, Color.b)
+	);
+}
+
+[numthreads(TILE_SIZE, TILE_SIZE, 1)]
+void MainCS(
+    uint2 GroupId : SV_GroupID,
+    uint2 GroupThreadId : SV_GroupThreadID,
+    uint GroupThreadIndex : SV_GroupIndex,
+    uint2 DispatchThreadId : SV_DispatchThreadID)
+{
+	float4 SvPosition = float4((float2)DispatchThreadId + Input_ViewportMin + 0.5f, 0.0f, 1.0f);
+	float2 UV = (SvPosition.xy) * Output_ExtentInverse;
+
+	if (IsComputeUVOutOfBounds(UV))
+	{
+		return;
+	}
+	
+	float4 Color = InputTexture.SampleLevel(InputSampler, UV, 0);
+	// Optionally linearize colors
+	Color.rgb = Decode(Color.rgb, InputEncoding);
+
+	// Apply keying
+	const float ColorTolerance = Params0.x;
+	const float EdgeSmoothness = Params0.y;
+
+#if USE_CLEAN_PLATE
+	float3 PlateColor = CleanPlateTexture.SampleLevel(CleanPlateSampler, UV, 0).rgb;
+	// Optionally linearize colors
+	PlateColor = Decode(PlateColor, CleanPlateEncoding);
+#else
+	float3 PlateColor = KeyColor;
+#endif
+
+	// Euclidean distance in (L)uv space as a measure of the color difference
+	const float Delta = CalcColorDelta(Color.rgb, PlateColor);
+	// Apply main delta key
+	Color.a = CalculateAlpha(ColorTolerance, ColorTolerance + EdgeSmoothness, Delta);
+
+	// Apply despill
+	if(DespillMethod > DESPILL_METHOD_NONE)
+	{
+		const float DespillStrength = Params1.x;
+		float3 Despilled;
+
+		if(DespillMethod == DESPILL_METHOD_GREEN_AVG)
+		{
+			Despilled = GreenDespillAvg(Color.rgb);
+		}
+		else if(DespillMethod == DESPILL_METHOD_BLUE_AVG)
+		{
+			Despilled = BlueDespillAvg(Color.rgb);
+		}
+		else if(DespillMethod == DESPILL_METHOD_ALPHA_KEY)
+		{
+			// Reduced delta such that the edge regions receive more despill (TODO: parametrize)
+			const float AdjustedDelta = 0.75 * Delta;
+			const float AlphaMask = 1.0 - CalculateAlpha(ColorTolerance, ColorTolerance + EdgeSmoothness, AdjustedDelta);
+			Despilled = AlphaKeyDespill(Color.rgb, AlphaMask);
+		}
+
+		// Re-insert luminance to prevent any darkening caused by despill
+		Despilled += Luminance(Color.rgb - Despilled);
+
+		Color.rgb = lerp(Color.rgb, Despilled, DespillStrength);
+	}
+
+	// Optionally re-encode colors
+	Color.rgb = Encode(Color.rgb, InputEncoding);
+	
+	uint2 PixelPos = DispatchThreadId + Output_ViewportMin;
+	RWOutputTexture[PixelPos] = Color;
+}
diff --git a/Engine/Plugins/Experimental/Compositing/Composite/Source/Composite/Private/Layers/CompositeLayerPlate.cpp b/Engine/Plugins/Experimental/Compositing/Composite/Source/Composite/Private/Layers/CompositeLayerPlate.cpp
--- a/Engine/Plugins/Experimental/Compositing/Composite/Source/Composite/Private/Layers/CompositeLayerPlate.cpp
+++ b/Engine/Plugins/Experimental/Compositing/Composite/Source/Composite/Private/Layers/CompositeLayerPlate.cpp
@@ -144,6 +144,11 @@ void UCompositeLayerPlate::PostLoad()
 
 	if (!HasAnyFlags(RF_ClassDefaultObject))
 	{
+		if (IsValid(Texture) && GetValidPreprocessingPassNum() > 0)
+		{
+			TextureProcessed = CreateTransientRenderTarget();
+		}
+
 		UpdateCompositeMeshes();
 	}
 }
@@ -264,28 +269,7 @@ void UCompositeLayerPlate::PostEditChangeProperty(FPropertyChangedEvent& Propert
 		{
 			if (IsValid(Texture) && !IsValid(TextureProcessed))
 			{
-				FIntPoint Resolution = GetTypedOuter<ACompositeActor>()->RenderResolution;
-
-#if 0 // Note: We may wish to expose options to retain the original texture resolution
-				if (const UMediaTexture* MediaTexture = Cast<UMediaTexture>(Texture))
-				{
-					Resolution = FIntPoint(MediaTexture->GetWidth(), MediaTexture->GetHeight());
-				}
-				else if (const UTexture2D* Texture2D = Cast<UTexture2D>(Texture))
-				{
-					Resolution = FIntPoint(Texture2D->GetSizeX(), Texture2D->GetSizeY());
-				}
-#endif
-
-				const FName UniqueName = MakeUniqueObjectName(this, UTextureRenderTarget2D::StaticClass(), FName("CompositeProcessedPlate"));
-				TextureProcessed = NewObject<UTextureRenderTarget2D>(this, UniqueName, RF_Transient);
-				TextureProcessed->SizeX = Resolution.X;
-				TextureProcessed->SizeY = Resolution.Y;
-				TextureProcessed->RenderTargetFormat = ETextureRenderTargetFormat::RTF_RGBA16f;
-				TextureProcessed->bForceLinearGamma = true;
-				TextureProcessed->bSupportsUAV = true;
-
-				TextureProcessed->UpdateResource();
+				TextureProcessed = CreateTransientRenderTarget();
 			}
 		}
 
@@ -544,6 +528,36 @@ int32 UCompositeLayerPlate::GetValidPreprocessingPassNum() const
 	return Count;
 }
 
+TObjectPtr<UTextureRenderTarget2D> UCompositeLayerPlate::CreateTransientRenderTarget()
+{
+	FIntPoint Resolution = GetTypedOuter<ACompositeActor>()->RenderResolution;
+
+#if 0 // Note: We may wish to expose options to retain the original texture resolution
+	if (const UMediaTexture* MediaTexture = Cast<UMediaTexture>(Texture))
+	{
+		Resolution = FIntPoint(MediaTexture->GetWidth(), MediaTexture->GetHeight());
+	}
+	else if (const UTexture2D* Texture2D = Cast<UTexture2D>(Texture))
+	{
+		Resolution = FIntPoint(Texture2D->GetSizeX(), Texture2D->GetSizeY());
+	}
+#endif
+
+	const FName UniqueName = MakeUniqueObjectName(this, UTextureRenderTarget2D::StaticClass(), FName("CompositeProcessedPlate"));
+	TObjectPtr<UTextureRenderTarget2D> RenderTarget = NewObject<UTextureRenderTarget2D>(this, UniqueName, RF_Transient);
+	RenderTarget->SizeX = Resolution.X;
+	RenderTarget->SizeY = Resolution.Y;
+	RenderTarget->RenderTargetFormat = ETextureRenderTargetFormat::RTF_RGBA16f;
+	RenderTarget->bForceLinearGamma = true;
+	RenderTarget->bSupportsUAV = true;
+	RenderTarget->AddressX = TA_Clamp;
+	RenderTarget->AddressY = TA_Clamp;
+
+	RenderTarget->UpdateResource();
+
+	return RenderTarget;
+}
+
 void UCompositeLayerPlate::UpdateCompositeMeshes(bool bRegisterOnly) const
 {
 	for (AActor* MeshActor : CompositeMeshes)
diff --git a/Engine/Plugins/Experimental/Compositing/Composite/Source/Composite/Private/Passes/CompositePassDeltaKeyer.cpp b/Engine/Plugins/Experimental/Compositing/Composite/Source/Composite/Private/Passes/CompositePassDeltaKeyer.cpp
new file mode 100644
--- /dev/null
+++ b/Engine/Plugins/Experimental/Compositing/Composite/Source/Composite/Private/Passes/CompositePassDeltaKeyer.cpp
@@ -0,0 +1,167 @@
+// Copyright Epic Games, Inc. All Rights Reserved.
+
+#include "Passes/CompositePassDeltaKeyer.h"
+
+#include "Engine/Texture.h"
+#include "Passes/CompositeCorePassProxy.h"
+#include "PixelShaderUtils.h"
+#include "TextureResource.h"
+
+DECLARE_GPU_STAT_NAMED(FCompositeDeltaKeyer, TEXT("Composite.DeltaKeyer"));
+
+class FCompositePassDeltaKeyerShader : public FGlobalShader
+{
+	DECLARE_GLOBAL_SHADER(FCompositePassDeltaKeyerShader);
+	SHADER_USE_PARAMETER_STRUCT(FCompositePassDeltaKeyerShader, FGlobalShader);
+
+	class FUseCleanPlate : SHADER_PERMUTATION_BOOL("USE_CLEAN_PLATE");
+	using FPermutationDomain = TShaderPermutationDomain<FUseCleanPlate>;
+
+	BEGIN_SHADER_PARAMETER_STRUCT(FParameters, )
+		SHADER_PARAMETER_RDG_TEXTURE(Texture2D, InputTexture)
+		SHADER_PARAMETER_SAMPLER(SamplerState, InputSampler)
+		SHADER_PARAMETER_TEXTURE(Texture2D, CleanPlateTexture)
+		SHADER_PARAMETER_SAMPLER(SamplerState, CleanPlateSampler)
+		SHADER_PARAMETER_RDG_TEXTURE_UAV(RWTexture2D<float4>, RWOutputTexture)
+		SHADER_PARAMETER_STRUCT(FScreenPassTextureViewportParameters, Input)
+		SHADER_PARAMETER_STRUCT(FScreenPassTextureViewportParameters, Output)
+		SHADER_PARAMETER(uint32, InputEncoding)
+		SHADER_PARAMETER(uint32, CleanPlateEncoding)
+		SHADER_PARAMETER(FUintVector2, InputTextureSize)
+		SHADER_PARAMETER(FVector3f, KeyColor)
+		SHADER_PARAMETER(FVector4f, Params0)
+		SHADER_PARAMETER(FVector4f, Params1)
+		SHADER_PARAMETER(uint32, DespillMethod)
+		SHADER_PARAMETER(uint32, bInvertAlpha)
+	END_SHADER_PARAMETER_STRUCT()
+
+	static inline void ModifyCompilationEnvironment(const FGlobalShaderPermutationParameters& Parameters, FShaderCompilerEnvironment& OutEnvironment)
+	{
+		OutEnvironment.SetDefine(TEXT("TILE_SIZE"), ThreadGroupSize);
+
+		FGlobalShader::ModifyCompilationEnvironment(Parameters, OutEnvironment);
+	}
+
+	static const uint32 ThreadGroupSize = 16;
+};
+IMPLEMENT_GLOBAL_SHADER(FCompositePassDeltaKeyerShader, "/Plugin/Composite/Private/CompositeDeltaKeyer.usf", "MainCS", SF_Compute);
+
+
+namespace UE
+{
+	namespace CompositeCore
+	{
+		class FCompositePassDeltaKeyerProxy : public FCompositeCorePassProxy
+		{
+		public:
+			using FCompositeCorePassProxy::FCompositeCorePassProxy;
+
+			FPassOutput Add(FRDGBuilder& GraphBuilder, const FSceneView& InView, const FPassInputArray& Inputs, const FPassContext& PassContext) const override
+			{
+				RDG_EVENT_SCOPE_STAT(GraphBuilder, FCompositeDeltaKeyer, "Composite.DeltaKeyer");
+				RDG_GPU_STAT_SCOPE(GraphBuilder, FCompositeDeltaKeyer);
+
+				const UE::CompositeCore::FResourceMetadata& Metadata = Inputs[0].Metadata;
+				const FScreenPassTexture& Input = Inputs[0].Texture;
+				FScreenPassRenderTarget Output = Inputs.OverrideOutput;
+
+				// If the override output is provided, it means that this is the last pass in post processing.
+				if (!Output.IsValid())
+				{
+					Output = CreateOutputRenderTarget(GraphBuilder, InView, PassContext.OutputViewRect, Input.Texture->Desc, TEXT("DeltaKeyerCompositePass"));
+				}
+
+				//CleanPlateWeakPtr.get()
+
+				FGlobalShaderMap* GlobalShaderMap = GetGlobalShaderMap(InView.GetFeatureLevel());
+				const FIntPoint TextureSize = Input.Texture->Desc.Extent;
+
+				const bool bSameSize = Output.ViewRect.Size() == Input.ViewRect.Size();
+				const FScreenPassTextureViewportParameters InputParameters = GetScreenPassTextureViewportParameters(FScreenPassTextureViewport(Input));
+				const FScreenPassTextureViewportParameters OutputParameters = GetScreenPassTextureViewportParameters(FScreenPassTextureViewport(Output));
+
+				FCompositePassDeltaKeyerShader::FParameters* PassParameters = GraphBuilder.AllocParameters<FCompositePassDeltaKeyerShader::FParameters>();
+				PassParameters->InputTexture = Input.Texture;
+				PassParameters->InputSampler = bSameSize ? TStaticSamplerState<SF_Point>::GetRHI() : TStaticSamplerState<SF_Bilinear>::GetRHI();
+				PassParameters->RWOutputTexture = GraphBuilder.CreateUAV(Output.Texture);
+				PassParameters->Input = InputParameters;
+				PassParameters->Output = OutputParameters;
+				PassParameters->InputEncoding = static_cast<uint32>(Inputs[0].Metadata.Encoding);
+				PassParameters->CleanPlateEncoding = CleanPlateEncoding;
+				PassParameters->KeyColor = FVector3f(KeyColor);
+				PassParameters->Params0 = Params0;
+				PassParameters->Params1 = Params1;
+				PassParameters->DespillMethod = DespillMethod;
+				PassParameters->bInvertAlpha = static_cast<uint32>(bInvertAlpha);
+	
+				bool bUseCleanPlate = false;
+				PassParameters->CleanPlateTexture = GBlackTexture->TextureRHI;
+				PassParameters->CleanPlateSampler = TStaticSamplerState<SF_Bilinear>::GetRHI();
+
+				if (CleanPlateWeakPtr.IsValid(false, true))
+				{
+					if (const FTextureResource* TextureResource = CleanPlateWeakPtr->GetResource())
+					{
+						PassParameters->CleanPlateTexture = TextureResource->TextureRHI;
+						bUseCleanPlate = true;
+					}
+				}
+
+				FCompositePassDeltaKeyerShader::FPermutationDomain PermutationVector;
+				PermutationVector.Set<FCompositePassDeltaKeyerShader::FUseCleanPlate>(bUseCleanPlate);
+
+				TShaderMapRef<FCompositePassDeltaKeyerShader> ComputeShader(GlobalShaderMap, PermutationVector);
+				FComputeShaderUtils::AddPass(
+					GraphBuilder,
+					RDG_EVENT_NAME("Composite.DeltaKeyer (%dx%d)", TextureSize.X, TextureSize.Y),
+					ERDGPassFlags::Compute,
+					ComputeShader,
+					PassParameters,
+					FComputeShaderUtils::GetGroupCount(TextureSize, FCompositePassDeltaKeyerShader::ThreadGroupSize)
+				);
+
+				return UE::CompositeCore::FPassOutput{ MoveTemp(Output), Metadata, PassOutputOverride };
+			}
+
+			FLinearColor KeyColor = FLinearColor::White;
+			FVector4f Params0 = FVector4f::Zero();
+			FVector4f Params1 = FVector4f::Zero();
+			uint32 DespillMethod = 0;
+			bool bInvertAlpha = false;
+			TWeakObjectPtr<UTexture> CleanPlateWeakPtr;
+			uint32 CleanPlateEncoding = 0;
+		};
+	}
+}
+
+UCompositePassDeltaKeyer::UCompositePassDeltaKeyer(const FObjectInitializer& ObjectInitializer)
+	: Super(ObjectInitializer)
+	, KeyColor(FLinearColor::Green)
+	, ColorTolerance(1.0f)
+	, EdgeSoftness(1.0f)
+	, AlphaClipRange(0.0f, 1.0f)
+	, DespillMethod(ECompositeDeltaKeyerDespillMethod::None)
+	, DespillStrength(1.0f)
+	, bInvertAlpha(false)
+{
+}
+
+UCompositePassDeltaKeyer::~UCompositePassDeltaKeyer() = default;
+
+bool UCompositePassDeltaKeyer::GetProxy(const UE::CompositeCore::FPassInputDecl& InputDecl, FSceneRenderingBulkObjectAllocator& InFrameAllocator, FCompositeCorePassProxy*& OutProxy) const
+{
+	using namespace UE::CompositeCore;
+
+	FCompositePassDeltaKeyerProxy* Proxy = InFrameAllocator.Create<FCompositePassDeltaKeyerProxy>(UE::CompositeCore::FPassInputDeclArray{ InputDecl });
+	Proxy->KeyColor = KeyColor;
+	Proxy->Params0 = FVector4f(ColorTolerance, EdgeSoftness, AlphaClipRange.X, AlphaClipRange.Y);
+	Proxy->Params1 = FVector4f(DespillStrength, 0.0f, 0.0f, 0.0f);
+	Proxy->DespillMethod = static_cast<uint32>(DespillMethod);
+	Proxy->bInvertAlpha = bInvertAlpha;
+	Proxy->CleanPlateWeakPtr = CleanPlate;
+	Proxy->CleanPlateEncoding = IsValid(CleanPlate) ? static_cast<uint32>(CleanPlate->GetGammaSpace()) : 0;
+
+	OutProxy = Proxy;
+	return true;
+}
+
diff --git a/Engine/Plugins/Experimental/Compositing/Composite/Source/Composite/Public/Layers/CompositeLayerPlate.h b/Engine/Plugins/Experimental/Compositing/Composite/Source/Composite/Public/Layers/CompositeLayerPlate.h
--- a/Engine/Plugins/Experimental/Compositing/Composite/Source/Composite/Public/Layers/CompositeLayerPlate.h
+++ b/Engine/Plugins/Experimental/Compositing/Composite/Source/Composite/Public/Layers/CompositeLayerPlate.h
@@ -13,6 +13,7 @@
 class UCompositeCoreSubsystem;
 class UPrimitiveComponent;
 class UTexture;
+class UTextureRenderTarget2D;
 
 UCLASS(MinimalAPI, BlueprintType, Blueprintable, EditInlineNew, CollapseCategories, meta = (DisplayName = "Plate"))
 class UCompositeLayerPlate : public UCompositeLayerBase, public FTickableGameObject
@@ -124,6 +125,9 @@ private:
 	/** Get number of valid pre-processing passes. */
 	int32 GetValidPreprocessingPassNum() const;
 
+	/** Convenience function to create a transient render target at the actor render resolution. */
+	TObjectPtr<UTextureRenderTarget2D> CreateTransientRenderTarget();
+
 private:
 #if WITH_EDITOR
 	/** Pre-edit list of registered composite meshes. */
@@ -132,7 +136,7 @@ private:
 
 	/** Processed texture render target. Note: Should be replaced by pool-managed RT. */
 	UPROPERTY(Transient)
-	TObjectPtr<class UTextureRenderTarget2D> TextureProcessed;
+	TObjectPtr<UTextureRenderTarget2D> TextureProcessed;
 };
 
 #undef UE_API
diff --git a/Engine/Plugins/Experimental/Compositing/Composite/Source/Composite/Public/Passes/CompositePassDeltaKeyer.h b/Engine/Plugins/Experimental/Compositing/Composite/Source/Composite/Public/Passes/CompositePassDeltaKeyer.h
new file mode 100644
--- /dev/null
+++ b/Engine/Plugins/Experimental/Compositing/Composite/Source/Composite/Public/Passes/CompositePassDeltaKeyer.h
@@ -0,0 +1,73 @@
+// Copyright Epic Games, Inc. All Rights Reserved.
+
+#pragma once
+
+#include "CompositePassBase.h"
+
+#include "CompositePassDeltaKeyer.generated.h"
+
+#define UE_API COMPOSITE_API
+
+UENUM(BlueprintType)
+enum class ECompositeDeltaKeyerDespillMethod : uint8
+{
+	None = 0,
+	GreenAverage = 1,
+	BlueAverage = 2,
+	Edge = 3,
+};
+
+/** Convenience pass to re-scale footage (with black bars) inside its container texture. */
+UCLASS(MinimalAPI, BlueprintType, Blueprintable, EditInlineNew, CollapseCategories, meta=(DisplayName="Delta Keyer Pass"))
+class UCompositePassDeltaKeyer : public UCompositePassBase
+{
+	GENERATED_BODY()
+
+public:
+	/** Constructor */
+	UE_API UCompositePassDeltaKeyer(const FObjectInitializer& ObjectInitializer);
+	
+	/** Destructor */
+	UE_API ~UCompositePassDeltaKeyer();
+
+	UE_API virtual bool GetProxy(const UE::CompositeCore::FPassInputDecl& InputDecl, FSceneRenderingBulkObjectAllocator& InFrameAllocator, FCompositeCorePassProxy*& OutProxy) const override;
+
+public:
+	/** Background key color. */
+	UPROPERTY(EditAnywhere, BlueprintReadWrite, Category = "", meta = (EditCondition = "CleanPlate == nullptr"))
+	FLinearColor KeyColor;
+
+	/**
+	 * Clean plate background for calculating color differences, replacing the static key color.
+	 * Resolution must match the composite actor render resolution.
+	 */
+	UPROPERTY(EditAnywhere, BlueprintReadWrite, Category = "")
+	TObjectPtr<UTexture> CleanPlate;
+
+	/**  A threshold for how similar a pixel’s color must be to the Key Color before it is keyed out. */
+	UPROPERTY(EditAnywhere, BlueprintReadWrite, Category = "", meta = (UIMin = 0.0f, UIMax = 2.0f))
+	float ColorTolerance;
+	
+	/** Controls how gradually the transition between keyed and non-keyed areas occurs at the edges. */
+	UPROPERTY(EditAnywhere, BlueprintReadWrite, Category = "", meta = (UIMin = 0.0f, UIMax = 2.0f))
+	float EdgeSoftness;
+
+	/** Clips any alpha value outside the specified range to zero or one respectively. */
+	UPROPERTY(EditAnywhere, BlueprintReadWrite, Category = "", meta = (UIMin = 0.0f, UIMax = 1.0f))
+	FVector2f AlphaClipRange;
+
+	/** Despill method, to remove/neutralize the residual key color spill on the subject. */
+	UPROPERTY(EditAnywhere, BlueprintReadWrite, Category = "")
+	ECompositeDeltaKeyerDespillMethod DespillMethod;
+
+	/** Strength of the spill reduction, 0.0: none, 1.0: full. */
+	UPROPERTY(EditAnywhere, BlueprintReadWrite, Category = "", meta = (UIMin = 0.0f, UIMax = 1.0f))
+	float DespillStrength;
+
+	/** Invert the alpha key. */
+	UPROPERTY(EditAnywhere, BlueprintReadWrite, Category = "")
+	bool bInvertAlpha;
+};
+
+#undef UE_API
+
diff --git a/Engine/Plugins/Experimental/Compositing/CompositeCore/Source/CompositeCore/Private/Passes/CompositeCorePassProxy.cpp b/Engine/Plugins/Experimental/Compositing/CompositeCore/Source/CompositeCore/Private/Passes/CompositeCorePassProxy.cpp
--- a/Engine/Plugins/Experimental/Compositing/CompositeCore/Source/CompositeCore/Private/Passes/CompositeCorePassProxy.cpp
+++ b/Engine/Plugins/Experimental/Compositing/CompositeCore/Source/CompositeCore/Private/Passes/CompositeCorePassProxy.cpp
@@ -114,7 +114,7 @@ FScreenPassRenderTarget FCompositeCorePassProxy::CreateOutputRenderTarget(FRDGBu
 	OutputDesc.Format	= FSceneTexturesConfig::Get().ColorFormat;
 	OutputDesc.NumMips	= 1;
 	OutputDesc.Depth	= 1;
-	OutputDesc.Flags	= TexCreate_RenderTargetable | TexCreate_ShaderResource; // TexCreate_UAV;
+	OutputDesc.Flags	= TexCreate_RenderTargetable | TexCreate_ShaderResource | TexCreate_UAV;
 	OutputDesc.Extent	= OutputViewRect.Size();
 
 	return FScreenPassRenderTarget(GraphBuilder.CreateTexture(OutputDesc, InName), OutputViewRect, InView.GetOverwriteLoadAction());
