Composite: Shader hotfix for uncaught compilation errors and potential division by zero.

diff --git a/Engine/Plugins/Experimental/Compositing/Composite/Shaders/Private/CompositeColorGrade.usf b/Engine/Plugins/Experimental/Compositing/Composite/Shaders/Private/CompositeColorGrade.usf
--- a/Engine/Plugins/Experimental/Compositing/Composite/Shaders/Private/CompositeColorGrade.usf
+++ b/Engine/Plugins/Experimental/Compositing/Composite/Shaders/Private/CompositeColorGrade.usf
@@ -1,7 +1,7 @@
 // Copyright Epic Games, Inc. All Rights Reserved.
 
 #include "/Engine/Private/Common.ush"
-#include "/Engine/Private/TonemaCommon.ush"
+#include "/Engine/Private/TonemapCommon.ush"
 #include "/Engine/Private/ACES/ACESCommon.ush"
 
 static const float DefaultTemperature = 6500.0;
@@ -106,7 +106,7 @@ void MainPS(
     BRANCH
 	if( abs(WhiteTemp - DefaultTemperature) > Epsilon || abs(WhiteTint - DefaultTint) > Epsilon )
 	{
-		LinearColor = WhiteBalance(LinearColor, WhiteTemp, WhiteTint, bIsTemperatureWhiteBalance, (float3x3)WorkingColorSpace.ToXYZ, (float3x3)WorkingColorSpace.FromXYZ);
+		LinearColor.rgb = WhiteBalance(LinearColor.rgb, WhiteTemp, WhiteTint, bIsTemperatureWhiteBalance, (float3x3) WorkingColorSpace.ToXYZ, (float3x3) WorkingColorSpace.FromXYZ);
 	}
 
     float4 ColorAP1 = mul(WorkingColorSpace.ToAP1, LinearColor);
diff --git a/Engine/Plugins/Experimental/Compositing/Composite/Shaders/Private/CompositeCommon.ush b/Engine/Plugins/Experimental/Compositing/Composite/Shaders/Private/CompositeCommon.ush
--- a/Engine/Plugins/Experimental/Compositing/Composite/Shaders/Private/CompositeCommon.ush
+++ b/Engine/Plugins/Experimental/Compositing/Composite/Shaders/Private/CompositeCommon.ush
@@ -28,7 +28,7 @@ float3 Encode(float3 InColor, uint InEncoding)
     BRANCH
     if(InEncoding == ESourceEncoding_Gamma)
     {
-        return pow(InColor, 2.2);
+        return pow(InColor, 1.0 / 2.2);
     }
     else if(InEncoding == ESourceEncoding_sRGB)
     {
diff --git a/Engine/Plugins/Experimental/Compositing/Composite/Shaders/Private/CompositeDeltaKeyer.usf b/Engine/Plugins/Experimental/Compositing/Composite/Shaders/Private/CompositeDeltaKeyer.usf
--- a/Engine/Plugins/Experimental/Compositing/Composite/Shaders/Private/CompositeDeltaKeyer.usf
+++ b/Engine/Plugins/Experimental/Compositing/Composite/Shaders/Private/CompositeDeltaKeyer.usf
@@ -47,6 +47,11 @@ static const float Lstar_CIELUV = 100.0;
 // Hand-picked scale factor used to adjust LUV color delta range
 static const float kDistanceScaleFactor_CIELUV = 5.0 * DeltaE_JND / Lstar_CIELUV;
 
+float SafeRcp(float x)
+{
+    return x > 0.0 ? rcp(x) : 0.0;
+}
+
 // Convert color from XYZ to CIE 1976 L*, u*, v* color space.
 // We use the latter for its improved perceptual uniformity, as a way to measure chromaticity distance for keying.
 float3 CIE_XYZ_to_LUV(float3 XYZ)
@@ -107,7 +112,8 @@ float CalculateAlpha(float Edge0, float Edge1, float Distance)
 
 float3 AlphaKeyDespill(float3 Color, float AlphaMask)
 {
-	float Weight = AlphaMask * dot(Color.rgb, KeyColor) / dot(KeyColor, KeyColor);
+    float KeyColorRcp = SafeRcp(dot(KeyColor, KeyColor));
+    float Weight = AlphaMask * dot(Color.rgb, KeyColor) * KeyColorRcp;
 
 	return max(0.0f, Color.rgb - Weight * KeyColor.rgb);
 }
