Composite: Various fixes/cleanup from recent CLs.

* Fixed composite mesh material update losing connection to textures when updated.
* Fixed plate layer sub-passes ordering.
* Fixed first & lass pass or layer index, to take into account invalid passes.
* Validate pass input arrays.
* Removed implicit decoding on keyer pass for now.

diff --git a/Engine/Plugins/Experimental/Compositing/Composite/Shaders/Private/CompositeDeltaKeyer.usf b/Engine/Plugins/Experimental/Compositing/Composite/Shaders/Private/CompositeDeltaKeyer.usf
--- a/Engine/Plugins/Experimental/Compositing/Composite/Shaders/Private/CompositeDeltaKeyer.usf
+++ b/Engine/Plugins/Experimental/Compositing/Composite/Shaders/Private/CompositeDeltaKeyer.usf
@@ -8,8 +8,6 @@
 #include "/Engine/Private/ScreenPass.ush"
 #include "/Engine/Private/PostProcessCommon.ush"
 
-#include "CompositeCommon.ush"
-
 #ifndef TILE_SIZE
     #error TILE_SIZE must be defined
 #endif
@@ -30,9 +28,6 @@ RWTexture2D<float4> RWOutputTexture;
 SCREEN_PASS_TEXTURE_VIEWPORT(Input)
 SCREEN_PASS_TEXTURE_VIEWPORT(Output)
 
-uint InputEncoding;
-uint CleanPlateEncoding;
-
 float3 KeyColor;
 float4 Params0;
 float4 Params1;
@@ -156,8 +151,6 @@ void MainCS(
 	}
 	
 	float4 Color = InputTexture.SampleLevel(InputSampler, UV, 0);
-	// Optionally linearize colors
-	Color.rgb = Decode(Color.rgb, InputEncoding);
 
 	// Apply keying
 	const float ColorTolerance = Params0.x;
@@ -165,8 +158,6 @@ void MainCS(
 
 #if USE_CLEAN_PLATE
 	float3 PlateColor = CleanPlateTexture.SampleLevel(CleanPlateSampler, UV, 0).rgb;
-	// Optionally linearize colors
-	PlateColor = Decode(PlateColor, CleanPlateEncoding);
 #else
 	float3 PlateColor = KeyColor;
 #endif
@@ -203,9 +194,6 @@ void MainCS(
 
 		Color.rgb = lerp(Color.rgb, Despilled, DespillStrength);
 	}
-
-	// Optionally re-encode colors
-	Color.rgb = Encode(Color.rgb, InputEncoding);
 	
 	uint2 PixelPos = DispatchThreadId + Output_ViewportMin;
 	RWOutputTexture[PixelPos] = Color;
diff --git a/Engine/Plugins/Experimental/Compositing/Composite/Source/Composite/Private/CompositeActor.cpp b/Engine/Plugins/Experimental/Compositing/Composite/Source/Composite/Private/CompositeActor.cpp
--- a/Engine/Plugins/Experimental/Compositing/Composite/Source/Composite/Private/CompositeActor.cpp
+++ b/Engine/Plugins/Experimental/Compositing/Composite/Source/Composite/Private/CompositeActor.cpp
@@ -74,6 +74,18 @@ namespace UE
 
 				// Iterate backwards so that lower passes in the UI are executed first
 				const int32 FirstPassIndex = InLayers.Num() - 1;
+
+				TOptional<int32> FirstValidPassIndex;
+				for (int32 LayerIndex = FirstPassIndex; LayerIndex >= 0; --LayerIndex)
+				{
+					const TObjectPtr<UCompositeLayerBase>& Layer = InLayers[LayerIndex];
+					if (IsValid(Layer) && Layer->bIsEnabled)
+					{
+						FirstValidPassIndex = LayerIndex;
+						break;
+					}
+				}
+
 				for (int32 LayerIndex = FirstPassIndex; LayerIndex >= 0; --LayerIndex)
 				{
 					const TObjectPtr<UCompositeLayerBase>& Layer = InLayers[LayerIndex];
@@ -82,7 +94,7 @@ namespace UE
 						continue;
 					}
 
-					Context.bIsFirstPass = (LayerIndex == FirstPassIndex);
+					Context.bIsFirstPass = (LayerIndex == FirstValidPassIndex);
 
 					FCompositeCorePassProxy* OutProxy;
 					if (Layer->GetProxy(Context, InFrameAllocator, OutProxy))
diff --git a/Engine/Plugins/Experimental/Compositing/Composite/Source/Composite/Private/Layers/CompositeLayerPlate.cpp b/Engine/Plugins/Experimental/Compositing/Composite/Source/Composite/Private/Layers/CompositeLayerPlate.cpp
--- a/Engine/Plugins/Experimental/Compositing/Composite/Source/Composite/Private/Layers/CompositeLayerPlate.cpp
+++ b/Engine/Plugins/Experimental/Compositing/Composite/Source/Composite/Private/Layers/CompositeLayerPlate.cpp
@@ -15,7 +15,9 @@
 #include "UObject/ConstructorHelpers.h"
 
 #if WITH_EDITOR
+#include "Editor.h"
 #include "UObject/UObjectGlobals.h"
+#include "UObject/ObjectSaveContext.h"
 #include "Framework/Notifications/NotificationManager.h"
 #include "Misc/TransactionObjectEvent.h"
 #include "Widgets/Notifications/SNotificationList.h"
@@ -73,6 +75,16 @@ UCompositeLayerPlate::UCompositeLayerPlate(const FObjectInitializer& ObjectIniti
 	Texture = Cast<UTexture>(DefaultMediaTexture.Object);
 	bHoldoutComposite = true;
 	Operation = ECompositeCoreMergeOp::Over;
+
+#if WITH_EDITOR
+	// If we are not the class default object...
+	if (!HasAnyFlags(RF_ClassDefaultObject))
+	{
+		// Hook into pre/post save.
+		FEditorDelegates::PreSaveWorldWithContext.AddUObject(this, &UCompositeLayerPlate::OnPreSaveWorld);
+		FEditorDelegates::PostSaveWorldWithContext.AddUObject(this, &UCompositeLayerPlate::OnPostSaveWorld);
+	}
+#endif // WITH_EDITOR
 }
 
 UCompositeLayerPlate::~UCompositeLayerPlate() = default;
@@ -153,6 +165,21 @@ void UCompositeLayerPlate::PostLoad()
 	}
 }
 
+void UCompositeLayerPlate::BeginDestroy()
+{
+	Super::BeginDestroy();
+
+#if WITH_EDITOR
+	// If we are not the class default object...
+	if (!HasAnyFlags(RF_ClassDefaultObject))
+	{
+		// Hook into pre/post save.
+		FEditorDelegates::PreSaveWorldWithContext.RemoveAll(this);
+		FEditorDelegates::PostSaveWorldWithContext.RemoveAll(this);
+	}
+#endif // WITH_EDITOR
+}
+
 const TArray<AActor*> UCompositeLayerPlate::GetCompositeMeshes() const
 {
 	return CompositeMeshes;
@@ -372,9 +399,29 @@ bool UCompositeLayerPlate::GetProxy(FTraversalContext& InContext, FSceneRenderin
 
 	if (TextureId.IsSet() && IsValid(TextureProcessed))
 	{
-		for (int32 Index=0; Index < Preprocessing.Num(); ++Index)
+		// Iterate backwards so that lower passes in the UI are executed first
+		const int32 FirstPassIndex = Preprocessing.Num() - 1;
+
+		TOptional<int32> FirstValidPassIndex;
+		TOptional<int32> LastValidPassIndex;
+		for (int32 PassIndex = FirstPassIndex; PassIndex >= 0; --PassIndex)
+		{
+			const TObjectPtr<UCompositePassBase>& Prepass = Preprocessing[PassIndex];
+
+			if (IsValid(Prepass) && Prepass->IsActive())
+			{
+				if (!FirstValidPassIndex.IsSet())
+				{
+					FirstValidPassIndex = PassIndex;
+				}
+
+				LastValidPassIndex = PassIndex;
+			}
+		}
+
+		for (int32 PassIndex = FirstPassIndex; PassIndex >= 0; --PassIndex)
 		{
-			const TObjectPtr<UCompositePassBase>& Prepass = Preprocessing[Index];
+			const TObjectPtr<UCompositePassBase>& Prepass = Preprocessing[PassIndex];
 
 			if (!IsValid(Prepass) || !Prepass->IsActive())
 			{
@@ -382,7 +429,7 @@ bool UCompositeLayerPlate::GetProxy(FTraversalContext& InContext, FSceneRenderin
 			}
 
 			FPassInputDecl PassInput;
-			if (Index == 0)
+			if (PassIndex == FirstValidPassIndex)
 			{
 				PassInput.Set<FPassExternalResourceDesc>({ TextureId.GetValue() });
 			}
@@ -394,8 +441,7 @@ bool UCompositeLayerPlate::GetProxy(FTraversalContext& InContext, FSceneRenderin
 			FCompositeCorePassProxy* PrepassProxy;
 			if (Prepass->GetProxy(PassInput, InFrameAllocator, PrepassProxy))
 			{
-				const bool bIsLastPass = (Index == Preprocessing.Num() - 1);
-				if (bIsLastPass)
+				if (PassIndex == LastValidPassIndex)
 				{
 					FExternalTexture& ExternalTexture = InContext.ExternalTextures.AddDefaulted_GetRef();
 					ExternalTexture.Texture = TextureProcessed;
@@ -528,9 +574,29 @@ int32 UCompositeLayerPlate::GetValidPreprocessingPassNum() const
 	return Count;
 }
 
+#if WITH_EDITOR
+void UCompositeLayerPlate::OnPreSaveWorld(UWorld* InWorld, FObjectPreSaveContext ObjectSaveContext)
+{
+	// We need to remove our asset user data before saving, as we do not need to save it out
+	// and only use it to know when the static mesh component changes.
+	UE::Composite::Private::RemoveCompositeAllAssetUserData(CompositeMeshes);
+}
+
+void UCompositeLayerPlate::OnPostSaveWorld(UWorld* InWorld, FObjectPostSaveContext ObjectSaveContext)
+{
+	constexpr bool bRegisterOnly = true;
+	UpdateCompositeMeshes(bRegisterOnly);
+}
+#endif
+
 TObjectPtr<UTextureRenderTarget2D> UCompositeLayerPlate::CreateTransientRenderTarget()
 {
-	FIntPoint Resolution = GetTypedOuter<ACompositeActor>()->RenderResolution;
+	FIntPoint Resolution = FIntPoint(1920, 1080);
+		
+	if (ACompositeActor* CompositeActor = GetTypedOuter<ACompositeActor>())
+	{
+		Resolution = CompositeActor->RenderResolution;
+	}
 
 #if 0 // Note: We may wish to expose options to retain the original texture resolution
 	if (const UMediaTexture* MediaTexture = Cast<UMediaTexture>(Texture))
diff --git a/Engine/Plugins/Experimental/Compositing/Composite/Source/Composite/Private/Passes/CompositePassCenteredScale.cpp b/Engine/Plugins/Experimental/Compositing/Composite/Source/Composite/Private/Passes/CompositePassCenteredScale.cpp
--- a/Engine/Plugins/Experimental/Compositing/Composite/Source/Composite/Private/Passes/CompositePassCenteredScale.cpp
+++ b/Engine/Plugins/Experimental/Compositing/Composite/Source/Composite/Private/Passes/CompositePassCenteredScale.cpp
@@ -40,6 +40,8 @@ namespace UE
 				RDG_EVENT_SCOPE_STAT(GraphBuilder, FCompositeCenteredScale, "Composite.CenteredScale");
 				RDG_GPU_STAT_SCOPE(GraphBuilder, FCompositeCenteredScale);
 
+				check(ValidateInputs(Inputs));
+
 				const UE::CompositeCore::FResourceMetadata& Metadata = Inputs[0].Metadata;
 				const FScreenPassTexture& Input = Inputs[0].Texture;
 				FScreenPassRenderTarget Output = Inputs.OverrideOutput;
diff --git a/Engine/Plugins/Experimental/Compositing/Composite/Source/Composite/Private/Passes/CompositePassColorGrade.cpp b/Engine/Plugins/Experimental/Compositing/Composite/Source/Composite/Private/Passes/CompositePassColorGrade.cpp
--- a/Engine/Plugins/Experimental/Compositing/Composite/Source/Composite/Private/Passes/CompositePassColorGrade.cpp
+++ b/Engine/Plugins/Experimental/Compositing/Composite/Source/Composite/Private/Passes/CompositePassColorGrade.cpp
@@ -7,6 +7,8 @@
 #include "PostProcess/PostProcessMaterialInputs.h"
 #include "Passes/CompositeCorePassProxy.h"
 
+DECLARE_GPU_STAT_NAMED(FCompositeColorGrade, TEXT("Composite.ColorGrade"));
+
 BEGIN_SHADER_PARAMETER_STRUCT(FColorGradePerRangeParameters, )
 	SHADER_PARAMETER(FVector4f, Saturation)
 	SHADER_PARAMETER(FVector4f, Contrast)
@@ -60,6 +62,11 @@ public:
 
 	UE::CompositeCore::FPassOutput Add(FRDGBuilder& GraphBuilder, const FSceneView& InView, const UE::CompositeCore::FPassInputArray& Inputs, const UE::CompositeCore::FPassContext& PassContext) const override
 	{
+		RDG_EVENT_SCOPE_STAT(GraphBuilder, FCompositeColorGrade, "Composite.ColorGrade");
+		RDG_GPU_STAT_SCOPE(GraphBuilder, FCompositeColorGrade);
+
+		check(ValidateInputs(Inputs));
+
 		const UE::CompositeCore::FResourceMetadata& Metadata = Inputs[0].Metadata;
 		const FScreenPassTexture& Input = Inputs[0].Texture;
 		FScreenPassRenderTarget Output = Inputs.OverrideOutput;
diff --git a/Engine/Plugins/Experimental/Compositing/Composite/Source/Composite/Private/Passes/CompositePassDeltaKeyer.cpp b/Engine/Plugins/Experimental/Compositing/Composite/Source/Composite/Private/Passes/CompositePassDeltaKeyer.cpp
--- a/Engine/Plugins/Experimental/Compositing/Composite/Source/Composite/Private/Passes/CompositePassDeltaKeyer.cpp
+++ b/Engine/Plugins/Experimental/Compositing/Composite/Source/Composite/Private/Passes/CompositePassDeltaKeyer.cpp
@@ -25,8 +25,6 @@ class FCompositePassDeltaKeyerShader : public FGlobalShader
 		SHADER_PARAMETER_RDG_TEXTURE_UAV(RWTexture2D<float4>, RWOutputTexture)
 		SHADER_PARAMETER_STRUCT(FScreenPassTextureViewportParameters, Input)
 		SHADER_PARAMETER_STRUCT(FScreenPassTextureViewportParameters, Output)
-		SHADER_PARAMETER(uint32, InputEncoding)
-		SHADER_PARAMETER(uint32, CleanPlateEncoding)
 		SHADER_PARAMETER(FUintVector2, InputTextureSize)
 		SHADER_PARAMETER(FVector3f, KeyColor)
 		SHADER_PARAMETER(FVector4f, Params0)
@@ -61,6 +59,8 @@ namespace UE
 				RDG_EVENT_SCOPE_STAT(GraphBuilder, FCompositeDeltaKeyer, "Composite.DeltaKeyer");
 				RDG_GPU_STAT_SCOPE(GraphBuilder, FCompositeDeltaKeyer);
 
+				check(ValidateInputs(Inputs));
+
 				const UE::CompositeCore::FResourceMetadata& Metadata = Inputs[0].Metadata;
 				const FScreenPassTexture& Input = Inputs[0].Texture;
 				FScreenPassRenderTarget Output = Inputs.OverrideOutput;
@@ -86,8 +86,6 @@ namespace UE
 				PassParameters->RWOutputTexture = GraphBuilder.CreateUAV(Output.Texture);
 				PassParameters->Input = InputParameters;
 				PassParameters->Output = OutputParameters;
-				PassParameters->InputEncoding = static_cast<uint32>(Inputs[0].Metadata.Encoding);
-				PassParameters->CleanPlateEncoding = CleanPlateEncoding;
 				PassParameters->KeyColor = FVector3f(KeyColor);
 				PassParameters->Params0 = Params0;
 				PassParameters->Params1 = Params1;
@@ -129,7 +127,6 @@ namespace UE
 			uint32 DespillMethod = 0;
 			bool bInvertAlpha = false;
 			TWeakObjectPtr<UTexture> CleanPlateWeakPtr;
-			uint32 CleanPlateEncoding = 0;
 		};
 	}
 }
@@ -159,7 +156,6 @@ bool UCompositePassDeltaKeyer::GetProxy(const UE::CompositeCore::FPassInputDecl&
 	Proxy->DespillMethod = static_cast<uint32>(DespillMethod);
 	Proxy->bInvertAlpha = bInvertAlpha;
 	Proxy->CleanPlateWeakPtr = CleanPlate;
-	Proxy->CleanPlateEncoding = IsValid(CleanPlate) ? static_cast<uint32>(CleanPlate->GetGammaSpace()) : 0;
 
 	OutProxy = Proxy;
 	return true;
diff --git a/Engine/Plugins/Experimental/Compositing/Composite/Source/Composite/Private/Passes/CompositePassMaterial.cpp b/Engine/Plugins/Experimental/Compositing/Composite/Source/Composite/Private/Passes/CompositePassMaterial.cpp
--- a/Engine/Plugins/Experimental/Compositing/Composite/Source/Composite/Private/Passes/CompositePassMaterial.cpp
+++ b/Engine/Plugins/Experimental/Compositing/Composite/Source/Composite/Private/Passes/CompositePassMaterial.cpp
@@ -9,6 +9,8 @@
 #include "ScreenPass.h"
 #include "PostProcess/PostProcessMaterialInputs.h"
 
+DECLARE_GPU_STAT_NAMED(FCompositeMaterial, TEXT("Composite.Material"));
+
 class FCompositePassMaterialProxy : public FCompositeCorePassProxy
 {
 public:
@@ -19,6 +21,11 @@ public:
 
 	UE::CompositeCore::FPassOutput Add(FRDGBuilder& GraphBuilder, const FSceneView& InView, const UE::CompositeCore::FPassInputArray& Inputs, const UE::CompositeCore::FPassContext& PassContext) const override
 	{
+		RDG_EVENT_SCOPE_STAT(GraphBuilder, FCompositeMaterial, "Composite.Material");
+		RDG_GPU_STAT_SCOPE(GraphBuilder, FCompositeMaterial);
+
+		check(ValidateInputs(Inputs));
+
 		const UE::CompositeCore::FResourceMetadata& Metadata = Inputs[0].Metadata;
 		FPostProcessMaterialInputs PostInputs = Inputs.ToPostProcessInputs(GraphBuilder, PassContext.SceneTextures);
 		// Specify the scene color output format in case a render target needs to be created and our input is render target incompatible.
diff --git a/Engine/Plugins/Experimental/Compositing/Composite/Source/Composite/Private/Passes/CompositePassOpenColorIO.cpp b/Engine/Plugins/Experimental/Compositing/Composite/Source/Composite/Private/Passes/CompositePassOpenColorIO.cpp
--- a/Engine/Plugins/Experimental/Compositing/Composite/Source/Composite/Private/Passes/CompositePassOpenColorIO.cpp
+++ b/Engine/Plugins/Experimental/Compositing/Composite/Source/Composite/Private/Passes/CompositePassOpenColorIO.cpp
@@ -6,6 +6,8 @@
 #include "ScreenPass.h"
 #include "PostProcess/PostProcessMaterialInputs.h"
 
+DECLARE_GPU_STAT_NAMED(FCompositeOpenColorIO, TEXT("Composite.OpenColorIO"));
+
 class FOpenColorIOCompositePassProxy
 	: public FCompositeCorePassProxy
 	, public TSharedFromThis<FOpenColorIOCompositePassProxy, ESPMode::ThreadSafe>
@@ -20,6 +22,11 @@ public:
 
 	UE::CompositeCore::FPassOutput Add(FRDGBuilder& GraphBuilder, const FSceneView& InView, const UE::CompositeCore::FPassInputArray& Inputs, const UE::CompositeCore::FPassContext& PassContext) const override
 	{
+		RDG_EVENT_SCOPE_STAT(GraphBuilder, FCompositeOpenColorIO, "Composite.OpenColorIO");
+		RDG_GPU_STAT_SCOPE(GraphBuilder, FCompositeOpenColorIO);
+
+		check(ValidateInputs(Inputs));
+
 		const UE::CompositeCore::FResourceMetadata& Metadata = Inputs[0].Metadata;
 		const FScreenPassTexture& Input = Inputs[0].Texture;
 		FScreenPassRenderTarget Output = Inputs.OverrideOutput;
diff --git a/Engine/Plugins/Experimental/Compositing/Composite/Source/Composite/Public/Layers/CompositeLayerPlate.h b/Engine/Plugins/Experimental/Compositing/Composite/Source/Composite/Public/Layers/CompositeLayerPlate.h
--- a/Engine/Plugins/Experimental/Compositing/Composite/Source/Composite/Public/Layers/CompositeLayerPlate.h
+++ b/Engine/Plugins/Experimental/Compositing/Composite/Source/Composite/Public/Layers/CompositeLayerPlate.h
@@ -13,6 +13,7 @@
 class UCompositeCoreSubsystem;
 class UPrimitiveComponent;
 class UTexture;
+class FObjectPreSaveContext;
 class UTextureRenderTarget2D;
 
 UCLASS(MinimalAPI, BlueprintType, Blueprintable, EditInlineNew, CollapseCategories, meta = (DisplayName = "Plate"))
@@ -50,8 +51,9 @@ public:
 
 protected:
 	//~ Begin UObject Interface
-	UE_API void PostInitProperties() override;
-	UE_API void PostLoad() override;
+	UE_API virtual void PostInitProperties() override;
+	UE_API virtual void PostLoad() override;
+	UE_API virtual void BeginDestroy() override;
 
 #if WITH_EDITOR
 	UE_API virtual bool CanEditChange(const FProperty* InProperty) const;
@@ -94,7 +96,7 @@ public:
 	TObjectPtr<UTexture> Texture;
 
 	/** Texture pre-processing passes, applied at the start of rendering and updating a render target. */
-	UPROPERTY(EditAnywhere, BlueprintReadWrite, Instanced, Category = "", meta = (DisplayName = "Pre-Processing", DisallowedClasses = "/Script/Composite.CompositePassMaterial"))
+	UPROPERTY(EditAnywhere, BlueprintReadWrite, Instanced, Category = "", meta = (DisplayName = "Pre-Processing", DisallowedClasses = "/Script/Composite.CompositePassMaterial, /Script/Composite.CompositePassFXAA"))
 	TArray<TObjectPtr<UCompositePassBase>> Preprocessing;
 
 	/** Sub-passes applied during post-processing on the texture or holdout composite render. */
@@ -132,6 +134,16 @@ private:
 #if WITH_EDITOR
 	/** Pre-edit list of registered composite meshes. */
 	TArray<TObjectPtr<AActor>> PreEditCompositeMeshes;
+
+	/**
+	 * Called before a level saves.
+	 */
+	void OnPreSaveWorld(UWorld* InWorld, FObjectPreSaveContext ObjectSaveContext);
+
+	/**
+	 * Called after a level has saved.
+	 */
+	void OnPostSaveWorld(UWorld* InWorld, FObjectPostSaveContext ObjectSaveContext);
 #endif
 
 	/** Processed texture render target. Note: Should be replaced by pool-managed RT. */
diff --git a/Engine/Plugins/Experimental/Compositing/Composite/Source/Composite/Public/Passes/CompositePassColorGrade.h b/Engine/Plugins/Experimental/Compositing/Composite/Source/Composite/Public/Passes/CompositePassColorGrade.h
--- a/Engine/Plugins/Experimental/Compositing/Composite/Source/Composite/Public/Passes/CompositePassColorGrade.h
+++ b/Engine/Plugins/Experimental/Compositing/Composite/Source/Composite/Public/Passes/CompositePassColorGrade.h
@@ -34,7 +34,10 @@ struct FCompositeTemperatureSettings
 
 };
 
-/** Pass to apply an OpenColorIO transform. */
+/**
+ * Color grade pass.
+ * Assumes input is in linear working color space.
+*/
 UCLASS(MinimalAPI, BlueprintType, Blueprintable, EditInlineNew, CollapseCategories, meta = (DisplayName = "Color Grade Pass"))
 class UCompositePassColorGrade : public UCompositePassBase
 {
diff --git a/Engine/Plugins/Experimental/Compositing/Composite/Source/Composite/Public/Passes/CompositePassDeltaKeyer.h b/Engine/Plugins/Experimental/Compositing/Composite/Source/Composite/Public/Passes/CompositePassDeltaKeyer.h
--- a/Engine/Plugins/Experimental/Compositing/Composite/Source/Composite/Public/Passes/CompositePassDeltaKeyer.h
+++ b/Engine/Plugins/Experimental/Compositing/Composite/Source/Composite/Public/Passes/CompositePassDeltaKeyer.h
@@ -17,7 +17,10 @@ enum class ECompositeDeltaKeyerDespillMethod : uint8
 	Edge = 3,
 };
 
-/** Convenience pass to re-scale footage (with black bars) inside its container texture. */
+/**
+ * Convenience pass to re-scale footage (with black bars) inside its container texture.
+ * Assumes input is in linear working color space.
+*/
 UCLASS(MinimalAPI, BlueprintType, Blueprintable, EditInlineNew, CollapseCategories, meta=(DisplayName="Delta Keyer Pass"))
 class UCompositePassDeltaKeyer : public UCompositePassBase
 {
diff --git a/Engine/Plugins/Experimental/Compositing/CompositeCore/Source/CompositeCore/Private/Passes/CompositeCorePassFXAAProxy.cpp b/Engine/Plugins/Experimental/Compositing/CompositeCore/Source/CompositeCore/Private/Passes/CompositeCorePassFXAAProxy.cpp
--- a/Engine/Plugins/Experimental/Compositing/CompositeCore/Source/CompositeCore/Private/Passes/CompositeCorePassFXAAProxy.cpp
+++ b/Engine/Plugins/Experimental/Compositing/CompositeCore/Source/CompositeCore/Private/Passes/CompositeCorePassFXAAProxy.cpp
@@ -9,6 +9,7 @@
 #include "PostProcess/PostProcessAA.h"
 #include "PostProcess/PostProcessMaterialInputs.h"
 
+DECLARE_GPU_STAT_NAMED(FCompositeCoreFXAA, TEXT("CompositeCore.FXAA"));
 DECLARE_GPU_STAT_NAMED(FCompositeCoreDisplayTransform, TEXT("CompositeCore.DisplayTransform"));
 
 class FCompositeCoreDisplayTransformShader : public FGlobalShader
@@ -84,6 +85,11 @@ namespace UE
 
 		FPassOutput FFXAAPassProxy::Add(FRDGBuilder& GraphBuilder, const FSceneView& InView, const FPassInputArray& Inputs, const FPassContext& PassContext) const
 		{
+			RDG_EVENT_SCOPE_STAT(GraphBuilder, FCompositeCoreFXAA, "CompositeCore.FXAA");
+			RDG_GPU_STAT_SCOPE(GraphBuilder, FCompositeCoreFXAA);
+
+			check(ValidateInputs(Inputs));
+
 			FScreenPassTexture Input = Inputs[0].Texture;
 			FResourceMetadata Metadata = Inputs[0].Metadata;
 			const bool bLinearSourceColors = (Metadata.Encoding == EEncoding::Linear);
diff --git a/Engine/Plugins/Experimental/Compositing/CompositeCore/Source/CompositeCore/Private/Passes/CompositeCorePassMergeProxy.cpp b/Engine/Plugins/Experimental/Compositing/CompositeCore/Source/CompositeCore/Private/Passes/CompositeCorePassMergeProxy.cpp
--- a/Engine/Plugins/Experimental/Compositing/CompositeCore/Source/CompositeCore/Private/Passes/CompositeCorePassMergeProxy.cpp
+++ b/Engine/Plugins/Experimental/Compositing/CompositeCore/Source/CompositeCore/Private/Passes/CompositeCorePassMergeProxy.cpp
@@ -117,6 +117,8 @@ namespace UE
 			RDG_EVENT_SCOPE_STAT(GraphBuilder, FCompositeCoreMerge, "CompositeCore.Merge");
 			RDG_GPU_STAT_SCOPE(GraphBuilder, FCompositeCoreMerge);
 
+			check(ValidateInputs(Inputs));
+
 			FGlobalShaderMap* GlobalShaderMap = GetGlobalShaderMap(InView.GetFeatureLevel());
 			const FSceneViewFamily* Family = InView.Family;
 			FResourceMetadata OutMetadata = {};
diff --git a/Engine/Plugins/Experimental/Compositing/CompositeCore/Source/CompositeCore/Private/Passes/CompositeCorePassProxy.cpp b/Engine/Plugins/Experimental/Compositing/CompositeCore/Source/CompositeCore/Private/Passes/CompositeCorePassProxy.cpp
--- a/Engine/Plugins/Experimental/Compositing/CompositeCore/Source/CompositeCore/Private/Passes/CompositeCorePassProxy.cpp
+++ b/Engine/Plugins/Experimental/Compositing/CompositeCore/Source/CompositeCore/Private/Passes/CompositeCorePassProxy.cpp
@@ -120,6 +120,11 @@ FScreenPassRenderTarget FCompositeCorePassProxy::CreateOutputRenderTarget(FRDGBu
 	return FScreenPassRenderTarget(GraphBuilder.CreateTexture(OutputDesc, InName), OutputViewRect, InView.GetOverwriteLoadAction());
 }
 
+bool FCompositeCorePassProxy::ValidateInputs(const UE::CompositeCore::FPassInputArray& Inputs) const
+{
+	return Inputs.Num() >= PassDeclaredInputs.Num();
+}
+
 UE::CompositeCore::FPassInputArray::FPassInputArray(
 	FRDGBuilder& GraphBuilder,
 	const FSceneView& InView,
diff --git a/Engine/Plugins/Experimental/Compositing/CompositeCore/Source/CompositeCore/Public/Passes/CompositeCorePassProxy.h b/Engine/Plugins/Experimental/Compositing/CompositeCore/Source/CompositeCore/Public/Passes/CompositeCorePassProxy.h
--- a/Engine/Plugins/Experimental/Compositing/CompositeCore/Source/CompositeCore/Public/Passes/CompositeCorePassProxy.h
+++ b/Engine/Plugins/Experimental/Compositing/CompositeCore/Source/CompositeCore/Public/Passes/CompositeCorePassProxy.h
@@ -318,6 +318,9 @@ public:
 
 protected:
 
+	/** Ensure that the received inputs match the declared inputs. */
+	UE_API bool ValidateInputs(const UE::CompositeCore::FPassInputArray& Inputs) const;
+
 	/** List of pass input types. */
 	UE::CompositeCore::FPassInputDeclArray PassDeclaredInputs;
 
