Composite: Trivial fix to prevent primitive registration after actor deletion.

diff --git a/Engine/Plugins/Experimental/Compositing/Composite/Source/Composite/Private/Layers/CompositeLayerPlate.cpp b/Engine/Plugins/Experimental/Compositing/Composite/Source/Composite/Private/Layers/CompositeLayerPlate.cpp
--- a/Engine/Plugins/Experimental/Compositing/Composite/Source/Composite/Private/Layers/CompositeLayerPlate.cpp
+++ b/Engine/Plugins/Experimental/Compositing/Composite/Source/Composite/Private/Layers/CompositeLayerPlate.cpp
@@ -96,12 +96,11 @@ void UCompositeLayerPlate::Tick(float DeltaTime)
 		return;
 	}
 
-	if (ACompositeActor* CompositeActor = GetTypedOuter<ACompositeActor>())
+	ACompositeActor* CompositeActor = GetTypedOuter<ACompositeActor>();
+	
+	if (!IsValid(CompositeActor) || !CompositeActor->IsActive())
 	{
-		if (!CompositeActor->IsActive())
-		{
-			return;
-		}
+		return;
 	}
 
 	UCompositeCoreSubsystem* Subsystem = UWorld::GetSubsystem<UCompositeCoreSubsystem>(GetWorld());
