Composite: Disable shadow casting by default on composite meshes.

diff --git a/Engine/Plugins/Experimental/Compositing/Composite/Source/Composite/Private/Components/CompositeMeshComponent.cpp b/Engine/Plugins/Experimental/Compositing/Composite/Source/Composite/Private/Components/CompositeMeshComponent.cpp
--- a/Engine/Plugins/Experimental/Compositing/Composite/Source/Composite/Private/Components/CompositeMeshComponent.cpp
+++ b/Engine/Plugins/Experimental/Compositing/Composite/Source/Composite/Private/Components/CompositeMeshComponent.cpp
@@ -27,6 +27,7 @@ UCompositeMeshComponent::UCompositeMeshComponent(const FObjectInitializer& Objec
 	{
 		SetStaticMesh(ConstructorStatics.BackgroundCube.Object);
 		SetMaterial(0, ConstructorStatics.DefaultMaterial.Object);
+		SetCastShadow(false);
 	}
 }
 
