Composite: Fix incorrect RDG pass name.

diff --git a/Engine/Plugins/Experimental/Compositing/Composite/Source/Composite/Private/Passes/CompositePassColorGrade.cpp b/Engine/Plugins/Experimental/Compositing/Composite/Source/Composite/Private/Passes/CompositePassColorGrade.cpp
--- a/Engine/Plugins/Experimental/Compositing/Composite/Source/Composite/Private/Passes/CompositePassColorGrade.cpp
+++ b/Engine/Plugins/Experimental/Compositing/Composite/Source/Composite/Private/Passes/CompositePassColorGrade.cpp
@@ -99,7 +99,7 @@ public:
 		TShaderMapRef<FCompositeColorGradeShader> PixelShader(GlobalShaderMap);
 		AddDrawScreenPass(
 			GraphBuilder,
-			RDG_EVENT_NAME("Composite.CenteredScale (%dx%d)", Viewport.Extent.X, Viewport.Extent.Y),
+			RDG_EVENT_NAME("Composite.ColorGrade (%dx%d)", Viewport.Extent.X, Viewport.Extent.Y),
 			InView,
 			Viewport,
 			Viewport,
