Composite: Adjust default shadow reflection capture settings to increase quality.

Added cost, but cheaper methods will be added later.

diff --git a/Engine/Plugins/Experimental/Compositing/Composite/Source/Composite/Private/Components/CompositeShadowReflectionCatcherComponent.cpp b/Engine/Plugins/Experimental/Compositing/Composite/Source/Composite/Private/Components/CompositeShadowReflectionCatcherComponent.cpp
--- a/Engine/Plugins/Experimental/Compositing/Composite/Source/Composite/Private/Components/CompositeShadowReflectionCatcherComponent.cpp
+++ b/Engine/Plugins/Experimental/Compositing/Composite/Source/Composite/Private/Components/CompositeShadowReflectionCatcherComponent.cpp
@@ -17,6 +17,8 @@ UCompositeShadowReflectionCatcherComponent::UCompositeShadowReflectionCatcherCom
 	bMainViewFamily = true;
 	bMainViewResolution = true;
 	bMainViewCamera = true;
+	// Inherit main view camera post-process settings to match main view DoF for example.
+	bInheritMainViewCameraPostProcessSettings = true;
 
 	ShowFlags.SetAtmosphere(false);
 	ShowFlags.SetAllowPrimitiveAlphaHoldout(false);
@@ -40,6 +42,10 @@ UCompositeShadowReflectionCatcherComponent::UCompositeShadowReflectionCatcherCom
 
 	PostProcessSettings.bOverride_LumenReflectionsScreenTraces = true;
 	PostProcessSettings.LumenReflectionsScreenTraces = false;
+
+	// While hit lighting is more expensive, it significantly minimizes artifacts & noise. Here we default with higher-quality settings.
+	PostProcessSettings.bOverride_LumenRayLightingMode = true;
+	PostProcessSettings.LumenRayLightingMode = ELumenRayLightingModeOverride::HitLighting;
 }
 
 UTextureRenderTarget2D* UCompositeShadowReflectionCatcherComponent::FindOrCreateRenderTarget(int32 InSizeX, int32 InSizeY)
