Composite: Add explicit distort/undistort pass support.

Currently, the layer merge pass handles distortion implicitely: this works well for media plate holdout composite, etc. However for the composite actor, we will want to undistort plates before projecting them into the scene, which this CL adds capability for.

Further work will be required for post-processing to operate on the unmodified plate.

diff --git a/Engine/Plugins/Experimental/Compositing/Composite/Composite.uplugin b/Engine/Plugins/Experimental/Compositing/Composite/Composite.uplugin
--- a/Engine/Plugins/Experimental/Compositing/Composite/Composite.uplugin
+++ b/Engine/Plugins/Experimental/Compositing/Composite/Composite.uplugin
@@ -35,6 +35,10 @@
 		{
 			"Name": "OpenColorIO",
 			"Enabled": true
+		},
+		{
+			"Name": "CameraCalibrationCore",
+			"Enabled": true
 		}
 	]
 }
\ No newline at end of file
diff --git a/Engine/Plugins/Experimental/Compositing/Composite/Shaders/Private/CompositeDistortion.usf b/Engine/Plugins/Experimental/Compositing/Composite/Shaders/Private/CompositeDistortion.usf
new file mode 100644
--- /dev/null
+++ b/Engine/Plugins/Experimental/Compositing/Composite/Shaders/Private/CompositeDistortion.usf
@@ -0,0 +1,49 @@
+// Copyright Epic Games, Inc. All Rights Reserved.
+
+#include "/Engine/Private/Common.ush"
+#include "/Engine/Private/LensDistortion.ush"
+#include "/Engine/Private/ScreenPass.ush"
+
+Texture2D InputTexture;
+SamplerState InputSampler;
+
+SCREEN_PASS_TEXTURE_VIEWPORT(Input)
+SCREEN_PASS_TEXTURE_VIEWPORT(Output)
+
+#define EDistortionUV_None        0
+#define EDistortionUV_Distorted   1
+#define EDistortionUV_Undistorted 2
+
+Texture2D<float2> DistortingDisplacementTexture;
+SamplerState DistortingDisplacementSampler;
+Texture2D<float2> UndistortingDisplacementTexture;
+SamplerState UndistortingDisplacementSampler;
+
+uint DistortionUV;
+
+float2 CalcInputUVs(float2 ViewportUV, float2 ViewportSize, float2 ViewportMin, float2 ExtentInverse, uint InDistortionUV)
+{
+    BRANCH
+    if(InDistortionUV == EDistortionUV_Distorted)
+    {
+        ViewportUV = ApplyLensDistortionOnViewportUV(DistortingDisplacementTexture, DistortingDisplacementSampler, ViewportUV);
+    }
+    else if(InDistortionUV == EDistortionUV_Undistorted)
+    {
+        ViewportUV = ApplyLensDistortionOnViewportUV(UndistortingDisplacementTexture, UndistortingDisplacementSampler, ViewportUV);
+    }
+
+    return (ViewportUV * ViewportSize + ViewportMin) * ExtentInverse;
+}
+
+void MainPS(
+    float4 SvPosition : SV_POSITION,
+    out float4 OutColor : SV_Target0
+    )
+{
+    float2 ViewportUV = (SvPosition.xy - Output_ViewportMin) * Output_ViewportSizeInverse;
+
+    float2 Texture_UV = CalcInputUVs(ViewportUV, Input_ViewportSize, Input_ViewportMin, Input_ExtentInverse, DistortionUV);
+
+	OutColor = InputTexture.SampleLevel(InputSampler, Texture_UV, 0);
+}
diff --git a/Engine/Plugins/Experimental/Compositing/Composite/Source/Composite/Composite.Build.cs b/Engine/Plugins/Experimental/Compositing/Composite/Source/Composite/Composite.Build.cs
--- a/Engine/Plugins/Experimental/Compositing/Composite/Source/Composite/Composite.Build.cs
+++ b/Engine/Plugins/Experimental/Compositing/Composite/Source/Composite/Composite.Build.cs
@@ -30,6 +30,7 @@ public class Composite : ModuleRules
 				"RenderCore",
 				"Renderer",
 				"RHI",
+				"CameraCalibrationCore"
 			}
 		);
 
diff --git a/Engine/Plugins/Experimental/Compositing/Composite/Source/Composite/Private/Passes/CompositePassColorGrade.cpp b/Engine/Plugins/Experimental/Compositing/Composite/Source/Composite/Private/Passes/CompositePassColorGrade.cpp
--- a/Engine/Plugins/Experimental/Compositing/Composite/Source/Composite/Private/Passes/CompositePassColorGrade.cpp
+++ b/Engine/Plugins/Experimental/Compositing/Composite/Source/Composite/Private/Passes/CompositePassColorGrade.cpp
@@ -74,7 +74,7 @@ public:
 		// If the override output is provided, it means that this is the last pass in post processing.
 		if (!Output.IsValid())
 		{
-			Output = CreateOutputRenderTarget(GraphBuilder, InView, PassContext.OutputViewRect, Input.Texture->Desc, TEXT("OcioCompositePass"));
+			Output = CreateOutputRenderTarget(GraphBuilder, InView, PassContext.OutputViewRect, Input.Texture->Desc, TEXT("ColorGradeCompositePass"));
 		}
 
 		FGlobalShaderMap* GlobalShaderMap = GetGlobalShaderMap(InView.GetFeatureLevel());
diff --git a/Engine/Plugins/Experimental/Compositing/Composite/Source/Composite/Private/Passes/CompositePassDistortion.cpp b/Engine/Plugins/Experimental/Compositing/Composite/Source/Composite/Private/Passes/CompositePassDistortion.cpp
new file mode 100644
--- /dev/null
+++ b/Engine/Plugins/Experimental/Compositing/Composite/Source/Composite/Private/Passes/CompositePassDistortion.cpp
@@ -0,0 +1,169 @@
+// Copyright Epic Games, Inc. All Rights Reserved.
+
+#include "Passes/CompositePassDistortion.h"
+
+#include "Camera/CameraComponent.h"
+#include "CameraCalibrationSubsystem.h"
+#include "LensDistortionSceneViewExtension.h"
+#include "ScreenPass.h"
+#include "SystemTextures.h"
+#include "Passes/CompositeCorePassProxy.h"
+#include "PixelShaderUtils.h"
+#include "PostProcess/LensDistortion.h"
+#include "PostProcess/PostProcessMaterialInputs.h"
+
+class FCompositeDistortionShader : public FGlobalShader
+{
+	DECLARE_GLOBAL_SHADER(FCompositeDistortionShader);
+	SHADER_USE_PARAMETER_STRUCT(FCompositeDistortionShader, FGlobalShader);
+
+	BEGIN_SHADER_PARAMETER_STRUCT(FParameters, )
+		SHADER_PARAMETER_RDG_TEXTURE(Texture2D, InputTexture)
+		SHADER_PARAMETER_SAMPLER(SamplerState, InputSampler)
+		SHADER_PARAMETER_STRUCT(FScreenPassTextureViewportParameters, Input)
+		SHADER_PARAMETER_STRUCT(FScreenPassTextureViewportParameters, Output)
+		SHADER_PARAMETER_RDG_TEXTURE(Texture2D, DistortingDisplacementTexture)
+		SHADER_PARAMETER_SAMPLER(SamplerState, DistortingDisplacementSampler)
+		SHADER_PARAMETER_RDG_TEXTURE(Texture2D, UndistortingDisplacementTexture)
+		SHADER_PARAMETER_SAMPLER(SamplerState, UndistortingDisplacementSampler)
+		SHADER_PARAMETER(uint32, DistortionUV)
+		RENDER_TARGET_BINDING_SLOTS()
+	END_SHADER_PARAMETER_STRUCT()
+};
+IMPLEMENT_GLOBAL_SHADER(FCompositeDistortionShader, "/Plugin/Composite/Private/CompositeDistortion.usf", "MainPS", SF_Pixel);
+
+DECLARE_GPU_STAT_NAMED(FCompositeDistortion, TEXT("Composite.Distortion"));
+
+class FDistortionCompositePassProxy : public FCompositeCorePassProxy
+{
+public:
+	enum class EDistortionUV
+	{
+		None = 0,
+		Distorted = 1,
+		Undistorted = 2
+	};
+
+	FDistortionCompositePassProxy(UE::CompositeCore::FPassInputDeclArray InPassDeclaredInputs)
+		: FCompositeCorePassProxy(MoveTemp(InPassDeclaredInputs))
+	{
+
+	}
+
+	UE::CompositeCore::FPassOutput Add(FRDGBuilder& GraphBuilder, const FSceneView& InView, const UE::CompositeCore::FPassInputArray& Inputs, const UE::CompositeCore::FPassContext& PassContext) const override
+	{
+		RDG_EVENT_SCOPE_STAT(GraphBuilder, FCompositeDistortion, "Composite.Distortion");
+		RDG_GPU_STAT_SCOPE(GraphBuilder, FCompositeDistortion);
+
+		check(ValidateInputs(Inputs));
+
+		UE::CompositeCore::FResourceMetadata Metadata = Inputs[0].Metadata;
+		const FScreenPassTexture& Input = Inputs[0].Texture;
+		FScreenPassRenderTarget Output = Inputs.OverrideOutput;
+
+		FLensDistortionLUT ViewDistortionLUT = LensDistortion::GetLUTUnsafe(InView);
+		bool bValidLensDistortion = ViewDistortionLUT.IsEnabled();
+
+		if (!bValidLensDistortion && LensDistortionSVE)
+		{
+			// Explicitely render view distortion LUT if existing view one doesn't yet exist, i.e. in pre-processing
+			bValidLensDistortion = LensDistortionSVE->RenderViewDistortionLUT(GraphBuilder, InView.ViewActor.ActorUniqueId, ViewDistortionLUT);
+		}
+
+		if (bValidLensDistortion)
+		{
+			// If the override output is provided, it means that this is the last pass in post processing.
+			if (!Output.IsValid())
+			{
+				Output = CreateOutputRenderTarget(GraphBuilder, InView, PassContext.OutputViewRect, Input.Texture->Desc, TEXT("DistortionCompositePass"));
+			}
+
+			const FScreenPassTextureViewportParameters InputParameters = GetScreenPassTextureViewportParameters(FScreenPassTextureViewport(Input));
+			const FScreenPassTextureViewportParameters OutputParameters = GetScreenPassTextureViewportParameters(FScreenPassTextureViewport(Output));
+
+			FGlobalShaderMap* GlobalShaderMap = GetGlobalShaderMap(InView.GetFeatureLevel());
+			FScreenPassTextureViewport Viewport = FScreenPassTextureViewport(Output);
+
+			FCompositeDistortionShader::FParameters* PassParameters = GraphBuilder.AllocParameters<FCompositeDistortionShader::FParameters>();
+			PassParameters->InputTexture = Input.Texture;
+			PassParameters->InputSampler = TStaticSamplerState<SF_Bilinear>::GetRHI();
+			PassParameters->Input = InputParameters;
+			PassParameters->Output = OutputParameters;
+			if (ViewDistortionLUT.IsEnabled())
+			{
+				PassParameters->DistortingDisplacementTexture = ViewDistortionLUT.DistortingDisplacementTexture;
+				PassParameters->UndistortingDisplacementTexture = ViewDistortionLUT.UndistortingDisplacementTexture;
+			}
+			else
+			{
+				PassParameters->UndistortingDisplacementTexture = GSystemTextures.GetBlackDummy(GraphBuilder);
+				PassParameters->DistortingDisplacementTexture = GSystemTextures.GetBlackDummy(GraphBuilder);
+			}
+			PassParameters->UndistortingDisplacementSampler = TStaticSamplerState<SF_Bilinear>::GetRHI();
+			PassParameters->DistortingDisplacementSampler = TStaticSamplerState<SF_Bilinear>::GetRHI();
+			// Sample input UVs with the opposite of the desired output distortion
+			switch (Distortion)
+			{
+			case ECompositeDistortion::Distort:
+				PassParameters->DistortionUV = static_cast<uint32>(EDistortionUV::Undistorted);
+				break;
+			case ECompositeDistortion::Undistort:
+				PassParameters->DistortionUV = static_cast<uint32>(EDistortionUV::Distorted);
+				break;
+			default:
+				checkNoEntry();
+				break;
+			}
+
+			PassParameters->RenderTargets[0] = Output.GetRenderTargetBinding();
+
+			TShaderMapRef<FCompositeDistortionShader> PixelShader(GlobalShaderMap);
+			
+			FPixelShaderUtils::AddFullscreenPass(
+				GraphBuilder,
+				GlobalShaderMap,
+				RDG_EVENT_NAME("Composite.Distortion (%dx%d)", Viewport.Extent.X, Viewport.Extent.Y),
+				PixelShader,
+				PassParameters,
+				Output.ViewRect
+			);
+
+			Metadata.bDistorted = (Distortion == ECompositeDistortion::Distort);
+			
+			return UE::CompositeCore::FPassOutput{ MoveTemp(Output), Metadata, PassOutputOverride };
+		}
+		else
+		{
+			return UE::CompositeCore::FPassOutput{ Output.IsValid() ? Output : Input, Metadata, PassOutputOverride };
+		}
+	}
+
+	TSharedPtr<FLensDistortionSceneViewExtension, ESPMode::ThreadSafe> LensDistortionSVE;
+	ECompositeDistortion Distortion;
+};
+
+UCompositePassDistortion::UCompositePassDistortion(const FObjectInitializer& ObjectInitializer)
+	: Super(ObjectInitializer)
+{
+	Distortion = ECompositeDistortion::Distort;
+}
+
+UCompositePassDistortion::~UCompositePassDistortion() = default;
+
+bool UCompositePassDistortion::GetProxy(const UE::CompositeCore::FPassInputDecl& InputDecl, FSceneRenderingBulkObjectAllocator& InFrameAllocator, FCompositeCorePassProxy*& OutProxy) const
+{
+	TSharedPtr<FLensDistortionSceneViewExtension, ESPMode::ThreadSafe> SVE;
+	UCameraCalibrationSubsystem* SubSystem = GEngine ? GEngine->GetEngineSubsystem<UCameraCalibrationSubsystem>() : nullptr;
+	if (!SubSystem)
+	{
+		return false;
+	}
+
+	FDistortionCompositePassProxy* Proxy = InFrameAllocator.Create<FDistortionCompositePassProxy>(UE::CompositeCore::FPassInputDeclArray{ InputDecl });
+	Proxy->Distortion = Distortion;
+	Proxy->LensDistortionSVE = SubSystem->GetSceneViewExtension();
+
+	OutProxy = Proxy;
+	return true;
+}
+
diff --git a/Engine/Plugins/Experimental/Compositing/Composite/Source/Composite/Public/Passes/CompositePassDistortion.h b/Engine/Plugins/Experimental/Compositing/Composite/Source/Composite/Public/Passes/CompositePassDistortion.h
new file mode 100644
--- /dev/null
+++ b/Engine/Plugins/Experimental/Compositing/Composite/Source/Composite/Public/Passes/CompositePassDistortion.h
@@ -0,0 +1,40 @@
+// Copyright Epic Games, Inc. All Rights Reserved.
+
+#pragma once
+
+#include "CoreTypes.h"
+#include "CompositePassBase.h"
+#include "Passes/CompositeCorePassMergeProxy.h"
+
+#include "CompositePassDistortion.generated.h"
+
+#define UE_API COMPOSITE_API
+
+UENUM(BlueprintType)
+enum class ECompositeDistortion : uint8
+{
+	Distort,
+	Undistort
+};
+
+/** Pass to apply an Distortion transform. */
+UCLASS(MinimalAPI, BlueprintType, Blueprintable, EditInlineNew, CollapseCategories, meta = (DisplayName = "Distortion Pass"))
+class UCompositePassDistortion : public UCompositePassBase
+{
+	GENERATED_BODY()
+
+public:
+	/** Constructor */
+	UE_API UCompositePassDistortion(const FObjectInitializer& ObjectInitializer);
+	/** Destructor */
+	UE_API ~UCompositePassDistortion();
+
+	UE_API virtual bool GetProxy(const UE::CompositeCore::FPassInputDecl& InputDecl, FSceneRenderingBulkObjectAllocator& InFrameAllocator, FCompositeCorePassProxy*& OutProxy) const override;
+
+public:
+	/** Distortion to apply on input. */
+	UPROPERTY(EditAnywhere, BlueprintReadWrite, Category = "")
+	ECompositeDistortion Distortion;
+};
+
+#undef UE_API
diff --git a/Engine/Plugins/Experimental/Compositing/CompositeCore/Source/CompositeCore/Private/Passes/CompositeCorePassProxy.cpp b/Engine/Plugins/Experimental/Compositing/CompositeCore/Source/CompositeCore/Private/Passes/CompositeCorePassProxy.cpp
--- a/Engine/Plugins/Experimental/Compositing/CompositeCore/Source/CompositeCore/Private/Passes/CompositeCorePassProxy.cpp
+++ b/Engine/Plugins/Experimental/Compositing/CompositeCore/Source/CompositeCore/Private/Passes/CompositeCorePassProxy.cpp
@@ -122,7 +122,7 @@ FScreenPassRenderTarget FCompositeCorePassProxy::CreateOutputRenderTarget(FRDGBu
 
 bool FCompositeCorePassProxy::ValidateInputs(const UE::CompositeCore::FPassInputArray& Inputs) const
 {
-	return Inputs.Num() >= PassDeclaredInputs.Num();
+	return Inputs.Num() >= FMath::Max(1, PassDeclaredInputs.Num());
 }
 
 UE::CompositeCore::FPassInputArray::FPassInputArray(
