Composite: Fix multiple post-processing passes on plate layer.

* We also ensure that newly added passes are added at the top of the stack.
* Additional crash fix on null processed texture attempted removal.

diff --git a/Engine/Plugins/Experimental/Compositing/Composite/Source/Composite/Private/Layers/CompositeLayerPlate.cpp b/Engine/Plugins/Experimental/Compositing/Composite/Source/Composite/Private/Layers/CompositeLayerPlate.cpp
--- a/Engine/Plugins/Experimental/Compositing/Composite/Source/Composite/Private/Layers/CompositeLayerPlate.cpp
+++ b/Engine/Plugins/Experimental/Compositing/Composite/Source/Composite/Private/Layers/CompositeLayerPlate.cpp
@@ -281,7 +281,7 @@ void UCompositeLayerPlate::PostEditChangeProperty(FPropertyChangedEvent& Propert
 
 		if (bChangeWasRemoval)
 		{
-			if (GetValidPreprocessingPassNum() == 0)
+			if (IsValid(TextureProcessed) && GetValidPreprocessingPassNum() == 0)
 			{
 				TextureProcessed->ConditionalBeginDestroy();
 				TextureProcessed->ReleaseResource();
@@ -290,8 +290,12 @@ void UCompositeLayerPlate::PostEditChangeProperty(FPropertyChangedEvent& Propert
 				CollectGarbage(GARBAGE_COLLECTION_KEEPFLAGS);
 			}
 		}
-		else if (PropertyChangedEvent.ChangeType == EPropertyChangeType::ArrayAdd ||
-			PropertyChangedEvent.ChangeType == EPropertyChangeType::ValueSet)
+		else if (PropertyChangedEvent.ChangeType == EPropertyChangeType::ArrayAdd)
+		{
+			// Always add elements at the start of the array (top of the stack)
+			Algo::Rotate(Preprocessing, Preprocessing.Num() - 1);
+		}
+		else if (PropertyChangedEvent.ChangeType == EPropertyChangeType::ValueSet)
 		{
 			if (IsValid(Texture) && !IsValid(TextureProcessed))
 			{
@@ -301,6 +305,14 @@ void UCompositeLayerPlate::PostEditChangeProperty(FPropertyChangedEvent& Propert
 
 		UpdateCompositeMeshes();
 	}
+	else if (PropertyName == GET_MEMBER_NAME_CHECKED(ThisClass, Postprocessing))
+	{
+		if (PropertyChangedEvent.ChangeType == EPropertyChangeType::ArrayAdd)
+		{
+			// Always add elements at the start of the array (top of the stack)
+			Algo::Rotate(Postprocessing, Postprocessing.Num() - 1);
+		}
+	}
 	else if (PropertyName == GET_MEMBER_NAME_CHECKED(ThisClass, CompositeMeshes))
 	{
 		const bool bChangeWasRemoval = PropertyChangedEvent.ChangeType == EPropertyChangeType::ArrayRemove
@@ -434,6 +446,7 @@ bool UCompositeLayerPlate::GetProxy(FTraversalContext& InContext, FSceneRenderin
 			}
 			else
 			{
+				// Default input from previous pass
 				PassInput.Set<FPassInternalResourceDesc>({});
 			}
 
@@ -487,6 +500,8 @@ bool UCompositeLayerPlate::GetProxy(FTraversalContext& InContext, FSceneRenderin
 		TArray<const FCompositeCorePassProxy*> SubPassProxies;
 		SubPassProxies.Reserve(Postprocessing.Num());
 
+		bool bFirstValidPass = true;
+
 		// Iterate backwards so that lower passes in the UI are executed first
 		for (int32 PassIndex = Postprocessing.Num() - 1; PassIndex >= 0; --PassIndex)
 		{
@@ -494,8 +509,20 @@ bool UCompositeLayerPlate::GetProxy(FTraversalContext& InContext, FSceneRenderin
 
 			if (IsValid(SubPass) && SubPass->IsActive())
 			{
+				FPassInputDecl PassInput;
+				if (bFirstValidPass)
+				{
+					PassInput = BaseInput0;
+					bFirstValidPass = false;
+				}
+				else
+				{
+					// Default input from previous pass
+					PassInput.Set<FPassInternalResourceDesc>({});
+				}
+
 				FCompositeCorePassProxy* SubPassProxy;
-				if (SubPass->GetProxy(BaseInput0, InFrameAllocator, SubPassProxy))
+				if (SubPass->GetProxy(PassInput, InFrameAllocator, SubPassProxy))
 				{
 					SubPassProxies.Add(SubPassProxy);
 				}
