Composite: Fixed invalid copy-pasted comment.

diff --git a/Engine/Plugins/Experimental/Compositing/Composite/Source/Composite/Public/Passes/CompositePassDeltaKeyer.h b/Engine/Plugins/Experimental/Compositing/Composite/Source/Composite/Public/Passes/CompositePassDeltaKeyer.h
--- a/Engine/Plugins/Experimental/Compositing/Composite/Source/Composite/Public/Passes/CompositePassDeltaKeyer.h
+++ b/Engine/Plugins/Experimental/Compositing/Composite/Source/Composite/Public/Passes/CompositePassDeltaKeyer.h
@@ -18,7 +18,7 @@ enum class ECompositeDeltaKeyerDespillMethod : uint8
 };
 
 /**
- * Convenience pass to re-scale footage (with black bars) inside its container texture.
+ * Chromaticity color keyer based on color distances in (L)UV color space, with spill removal support.
  * Assumes input is in linear working color space.
 */
 UCLASS(MinimalAPI, BlueprintType, Blueprintable, EditInlineNew, CollapseCategories, meta=(DisplayName="Delta Keyer Pass"))
