Composite: Auto-configure primitive properties for shadow & reflection catching on actors when registered.

* Hidden in scene capture (while preserving shadows & indirect) creates a cleaner shadow/reflection multiplicative matte that avoids edge artifacts, especially with DoF.
* Users can opt-out of this behavior by disabling "Auto-Configure Actors", for instance if they want to preserve screenspace AO.

diff --git a/Engine/Plugins/Experimental/Compositing/Composite/Source/Composite/Private/Layers/CompositeLayerShadowReflection.cpp b/Engine/Plugins/Experimental/Compositing/Composite/Source/Composite/Private/Layers/CompositeLayerShadowReflection.cpp
--- a/Engine/Plugins/Experimental/Compositing/Composite/Source/Composite/Private/Layers/CompositeLayerShadowReflection.cpp
+++ b/Engine/Plugins/Experimental/Compositing/Composite/Source/Composite/Private/Layers/CompositeLayerShadowReflection.cpp
@@ -3,12 +3,14 @@
 #include "Layers/CompositeLayerShadowReflection.h"
 
 #include "CompositeActor.h"
+#include "Components/PrimitiveComponent.h"
 #include "Components/CompositeShadowReflectionCatcherComponent.h"
 
 UCompositeLayerShadowReflection::UCompositeLayerShadowReflection(const FObjectInitializer& ObjectInitializer)
 	: Super(ObjectInitializer)
 {
 	Operation = ECompositeCoreMergeOp::Multiply;
+	bAutoConfigureActors = true;
 }
 
 UCompositeLayerShadowReflection::~UCompositeLayerShadowReflection() = default;
@@ -18,7 +20,9 @@ void UCompositeLayerShadowReflection::PostEditChangeProperty(FPropertyChangedEve
 {
 	Super::PostEditChangeProperty(PropertyChangedEvent);
 
-	if (PropertyChangedEvent.GetMemberPropertyName() == GET_MEMBER_NAME_CHECKED(ThisClass, bIsEnabled))
+	const FName PropertyName = PropertyChangedEvent.GetMemberPropertyName();
+
+	if (PropertyName == GET_MEMBER_NAME_CHECKED(ThisClass, bIsEnabled))
 	{
 		ACompositeActor* CompositeActor = GetTypedOuter<ACompositeActor>();
 		if (IsValid(CompositeActor))
@@ -26,6 +30,13 @@ void UCompositeLayerShadowReflection::PostEditChangeProperty(FPropertyChangedEve
 			CompositeActor->SetReflectionShadowCatcherCaptureEveryFrame(this, bIsEnabled);
 		}
 	}
+	else if (PropertyName == GET_MEMBER_NAME_CHECKED(ThisClass, Actors))
+	{
+		if (bAutoConfigureActors)
+		{
+			ConfigurePrimitivesForShadowReflectionCatching();
+		}
+	}
 }
 #endif
 
@@ -96,3 +107,34 @@ bool UCompositeLayerShadowReflection::GetProxy(FTraversalContext& InContext, FSc
 	return true;
 }
 
+const TArray<AActor*> UCompositeLayerShadowReflection::GetActors() const
+{
+	return Actors;
+}
+
+void UCompositeLayerShadowReflection::SetActors(TArray<AActor*> InActors)
+{
+	Actors = MoveTemp(InActors);
+
+	if (bAutoConfigureActors)
+	{
+		ConfigurePrimitivesForShadowReflectionCatching();
+	}
+}
+
+void UCompositeLayerShadowReflection::ConfigurePrimitivesForShadowReflectionCatching()
+{
+	for (AActor* Actor : Actors)
+	{
+		if (IsValid(Actor))
+		{
+			if (UPrimitiveComponent* PrimitiveComponent = Actor->GetComponentByClass<UPrimitiveComponent>())
+			{
+				PrimitiveComponent->SetHiddenInSceneCapture(true);
+				PrimitiveComponent->SetAffectIndirectLightingWhileHidden(true);
+				PrimitiveComponent->SetCastHiddenShadow(true);
+			}
+		}
+	}
+}
+
diff --git a/Engine/Plugins/Experimental/Compositing/Composite/Source/Composite/Public/Layers/CompositeLayerPlate.h b/Engine/Plugins/Experimental/Compositing/Composite/Source/Composite/Public/Layers/CompositeLayerPlate.h
--- a/Engine/Plugins/Experimental/Compositing/Composite/Source/Composite/Public/Layers/CompositeLayerPlate.h
+++ b/Engine/Plugins/Experimental/Compositing/Composite/Source/Composite/Public/Layers/CompositeLayerPlate.h
@@ -65,11 +65,11 @@ protected:
 
 public:
 
-	/** Get the enabled state. */
+	/** Get composite meshes. */
 	UFUNCTION(BlueprintGetter)
 	UE_API const TArray<AActor*> GetCompositeMeshes() const;
 
-	/** Set the enabled state. */
+	/** Set composite meshes. */
 	UFUNCTION(BlueprintSetter)
 	UE_API void SetCompositeMeshes(TArray<AActor*> InCompositeMeshes);
 
diff --git a/Engine/Plugins/Experimental/Compositing/Composite/Source/Composite/Public/Layers/CompositeLayerShadowReflection.h b/Engine/Plugins/Experimental/Compositing/Composite/Source/Composite/Public/Layers/CompositeLayerShadowReflection.h
--- a/Engine/Plugins/Experimental/Compositing/Composite/Source/Composite/Public/Layers/CompositeLayerShadowReflection.h
+++ b/Engine/Plugins/Experimental/Compositing/Composite/Source/Composite/Public/Layers/CompositeLayerShadowReflection.h
@@ -31,8 +31,30 @@ public:
 	/** Getter function to override, returning pass layer proxies to be passed to the render thread. (Proxy objects should be allocated from the provided allocator.) */
 	UE_API virtual bool GetProxy(FTraversalContext& InContext, FSceneRenderingBulkObjectAllocator& InFrameAllocator, FCompositeCorePassProxy*& OutProxy) const override;
 
-	/** Also provide the post-processing composite graph to SSR input. */
-	UPROPERTY(EditAnywhere, BlueprintReadWrite, Category = "")
+	/** Get reflection & shadow caster actors. */
+	UFUNCTION(BlueprintGetter)
+	UE_API const TArray<AActor*> GetActors() const;
+
+	/** Set reflection & shadow caster actors. */
+	UFUNCTION(BlueprintSetter)
+	UE_API void SetActors(TArray<AActor*> InActors);
+
+	/**
+	 * When enabled, we automatically toggle the following (primitive component) properties on actors when registered:
+	 * 
+	 * * bHiddenInSceneCapture: true
+	 * * bAffectIndirectLightingWhileHidden: true
+	 * * bCastHiddenShadow: true
+	*/
+	UPROPERTY(EditAnywhere, BlueprintReadWrite, Category = "Composite", AdvancedDisplay, meta = (DisplayName = "Auto-Configure Actors"))
+	bool bAutoConfigureActors;
+
+private:
+	/** Convenience function to configure actor primitives. */
+	void ConfigurePrimitivesForShadowReflectionCatching();
+
+	/** List of reflection & shadow caster actors. */
+	UPROPERTY(EditAnywhere, BlueprintGetter = GetActors, BlueprintSetter = SetActors, Category = "Composite")
 	TArray<TObjectPtr<AActor>> Actors;
 };
 
