Composite: Renamed composite mesh materials since they no longer are responsible for keying.

* Also updated unlit material blend mode from Translucent to (premultipled) Alpha Composite for accurate blending with primitive alpha holdout mask.
* Also updated keyer pass to output premult alpha.

#virtualized

diff --git a/Engine/Plugins/Experimental/Compositing/Composite/Shaders/Private/CompositeDeltaKeyer.usf b/Engine/Plugins/Experimental/Compositing/Composite/Shaders/Private/CompositeDeltaKeyer.usf
--- a/Engine/Plugins/Experimental/Compositing/Composite/Shaders/Private/CompositeDeltaKeyer.usf
+++ b/Engine/Plugins/Experimental/Compositing/Composite/Shaders/Private/CompositeDeltaKeyer.usf
@@ -190,10 +190,13 @@ void MainCS(
 		}
 
 		// Re-insert luminance to prevent any darkening caused by despill
-		Despilled += Luminance(Color.rgb - Despilled);
+		Despilled += Luminance(max(0.0f, Color.rgb - Despilled));
 
 		Color.rgb = lerp(Color.rgb, Despilled, DespillStrength);
 	}
+
+	// Return premultiplied alpha for AlphaComposite material blend mode
+	Color.rgb *= Color.a;
 	
 	uint2 PixelPos = DispatchThreadId + Output_ViewportMin;
 	RWOutputTexture[PixelPos] = Color;
diff --git a/Engine/Plugins/Experimental/Compositing/Composite/Source/Composite/Private/Components/CompositeMeshComponent.cpp b/Engine/Plugins/Experimental/Compositing/Composite/Source/Composite/Private/Components/CompositeMeshComponent.cpp
--- a/Engine/Plugins/Experimental/Compositing/Composite/Source/Composite/Private/Components/CompositeMeshComponent.cpp
+++ b/Engine/Plugins/Experimental/Compositing/Composite/Source/Composite/Private/Components/CompositeMeshComponent.cpp
@@ -17,7 +17,7 @@ UCompositeMeshComponent::UCompositeMeshComponent(const FObjectInitializer& Objec
 		ConstructorHelpers::FObjectFinder<UMaterial> DefaultMaterial;
 		FConstructorStatics()
 			: BackgroundCube(TEXT("/Engine/EngineMeshes/BackgroundCube.BackgroundCube"))
-			, DefaultMaterial(TEXT("/Composite/Materials/M_Composite_KeyedProjection.M_Composite_KeyedProjection"))
+			, DefaultMaterial(TEXT("/Composite/Materials/M_CompositeMesh_Lit_Masked.M_CompositeMesh_Lit_Masked"))
 		{
 		}
 	};
