Composite: Update default composite mesh.

* Default material to Lit masked so that users can simply enable the shadow & reflection layer.
* Keying setups will now require a material change to M_CompositeMesh_Unlit_AlphaComposite.

#virtualized

diff --git a/Engine/Plugins/Experimental/Compositing/Composite/Source/Composite/Private/Components/CompositeMeshComponent.cpp b/Engine/Plugins/Experimental/Compositing/Composite/Source/Composite/Private/Components/CompositeMeshComponent.cpp
--- a/Engine/Plugins/Experimental/Compositing/Composite/Source/Composite/Private/Components/CompositeMeshComponent.cpp
+++ b/Engine/Plugins/Experimental/Compositing/Composite/Source/Composite/Private/Components/CompositeMeshComponent.cpp
@@ -13,20 +13,17 @@ UCompositeMeshComponent::UCompositeMeshComponent(const FObjectInitializer& Objec
 {
 	struct FConstructorStatics
 	{
-		ConstructorHelpers::FObjectFinder<UStaticMesh> BackgroundCube;
-		ConstructorHelpers::FObjectFinder<UMaterial> DefaultMaterial;
+		ConstructorHelpers::FObjectFinder<UStaticMesh> DefaultMesh;
 		FConstructorStatics()
-			: BackgroundCube(TEXT("/Engine/EngineMeshes/BackgroundCube.BackgroundCube"))
-			, DefaultMaterial(TEXT("/Composite/Materials/M_CompositeMesh_Lit_Masked.M_CompositeMesh_Lit_Masked"))
+			: DefaultMesh(TEXT("/Composite/Meshes/SM_DefaultCompositeMesh.SM_DefaultCompositeMesh"))
 		{
 		}
 	};
 
 	static FConstructorStatics ConstructorStatics;
-	if (ensure(ConstructorStatics.BackgroundCube.Object != nullptr))
+	if (ensure(ConstructorStatics.DefaultMesh.Object != nullptr))
 	{
-		SetStaticMesh(ConstructorStatics.BackgroundCube.Object);
-		SetMaterial(0, ConstructorStatics.DefaultMaterial.Object);
+		SetStaticMesh(ConstructorStatics.DefaultMesh.Object);
 		SetCastShadow(false);
 	}
 }
