Composite: Add scene capture layer type.

* New layer type adding convenience for creating & integrating scene captures or custom render passes into the composite.
* Refactored composite actor scene capture management & updated shadow/reflection catcher.
* Fixed issue where replacing a scene-capture managing layer would not delete the components.

diff --git a/Engine/Plugins/Experimental/Compositing/Composite/Source/Composite/Private/Components/CompositeShadowReflectionCatcherComponent.cpp b/Engine/Plugins/Experimental/Compositing/Composite/Source/Composite/Private/Components/CompositeShadowReflectionCatcherComponent.cpp
--- a/Engine/Plugins/Experimental/Compositing/Composite/Source/Composite/Private/Components/CompositeShadowReflectionCatcherComponent.cpp
+++ b/Engine/Plugins/Experimental/Compositing/Composite/Source/Composite/Private/Components/CompositeShadowReflectionCatcherComponent.cpp
@@ -2,6 +2,7 @@
 
 #include "Components/CompositeShadowReflectionCatcherComponent.h"
 
+#include "CompositeUtilities.h"
 #include "Engine/TextureRenderTarget2D.h"
 #include "UObject/Package.h"
 
@@ -10,15 +11,7 @@
 UCompositeShadowReflectionCatcherComponent::UCompositeShadowReflectionCatcherComponent(const FObjectInitializer& ObjectInitializer)
 	: Super(ObjectInitializer)
 {
-	PrimitiveRenderMode = ESceneCapturePrimitiveRenderMode::PRM_RenderScenePrimitives;
-	CaptureSource = ESceneCaptureSource::SCS_FinalColorHDR;
-
-	bUseRayTracingIfEnabled = true;
-	bMainViewFamily = true;
-	bMainViewResolution = true;
-	bMainViewCamera = true;
-	// Inherit main view camera post-process settings to match main view DoF for example.
-	bInheritMainViewCameraPostProcessSettings = true;
+	UE::Composite::ConfigureSceneCapture(*this, ESceneCapturePrimitiveRenderMode::PRM_RenderScenePrimitives);
 
 	ShowFlags.SetAtmosphere(false);
 	ShowFlags.SetAllowPrimitiveAlphaHoldout(false);
@@ -30,22 +23,6 @@ UCompositeShadowReflectionCatcherComponent::UCompositeShadowReflectionCatcherCom
 	ShowFlags.SetScreenSpaceReflections(false);
 	//Note: This can cause temporal ghosting artifacts on moving shadows, but remains preferable over the otherwise default shadow noise.
 	ShowFlags.SetTemporalAA(true);
-
-	PostProcessSettings.bOverride_ReflectionMethod = true;
-	PostProcessSettings.ReflectionMethod = EReflectionMethod::Lumen;
-
-	PostProcessSettings.bOverride_DynamicGlobalIlluminationMethod = true;
-	PostProcessSettings.DynamicGlobalIlluminationMethod = EDynamicGlobalIlluminationMethod::Lumen;
-
-	PostProcessSettings.bOverride_LumenFinalGatherScreenTraces = true;
-	PostProcessSettings.LumenFinalGatherScreenTraces = false;
-
-	PostProcessSettings.bOverride_LumenReflectionsScreenTraces = true;
-	PostProcessSettings.LumenReflectionsScreenTraces = false;
-
-	// While hit lighting is more expensive, it significantly minimizes artifacts & noise. Here we default with higher-quality settings.
-	PostProcessSettings.bOverride_LumenRayLightingMode = true;
-	PostProcessSettings.LumenRayLightingMode = ELumenRayLightingModeOverride::HitLighting;
 }
 
 UTextureRenderTarget2D* UCompositeShadowReflectionCatcherComponent::FindOrCreateRenderTarget(int32 InSizeX, int32 InSizeY)
diff --git a/Engine/Plugins/Experimental/Compositing/Composite/Source/Composite/Private/CompositeActor.cpp b/Engine/Plugins/Experimental/Compositing/Composite/Source/Composite/Private/CompositeActor.cpp
--- a/Engine/Plugins/Experimental/Compositing/Composite/Source/Composite/Private/CompositeActor.cpp
+++ b/Engine/Plugins/Experimental/Compositing/Composite/Source/Composite/Private/CompositeActor.cpp
@@ -5,6 +5,7 @@
 #include "CompositeCoreSubsystem.h"
 #include "CompositeModule.h"
 #include "CompositeCoreSettings.h"
+#include "CompositeUtilities.h"
 
 #include "Components/CompositeShadowReflectionCatcherComponent.h"
 #include "Components/CompositeViewProjectionComponent.h"
@@ -22,6 +23,7 @@
 #include "Kismet/GameplayStatics.h"
 #include "Materials/MaterialInstanceDynamic.h"
 #include "SceneView.h"
+#include "UObject/FortniteMainBranchObjectVersion.h"
 #include "UObject/ObjectSaveContext.h"
 #include "UObject/UObjectIterator.h"
 
@@ -42,6 +44,9 @@ namespace UE
 {
 	namespace Composite
 	{
+		/** Tag that will be added on every component generated by the composite actor. */
+		const FLazyName DefaultCompositeTag = TEXT("Composite Generated Component");
+
 		namespace Private
 		{
 			/** Return solo layer from list. */
@@ -142,18 +147,6 @@ namespace UE
 
 				return RenderWork;
 			}
-
-			void SetCaptureEveryFrame(TMap<TWeakObjectPtr<const UCompositeLayerBase>, TObjectPtr<UCompositeShadowReflectionCatcherComponent>>& InComponents, bool bState)
-			{
-				for (auto& Pair : InComponents)
-				{
-					TObjectPtr<UCompositeShadowReflectionCatcherComponent>& Component = Pair.Value;
-					if (IsValid(Component))
-					{
-						Component->bCaptureEveryFrame = bState;
-					}
-				}
-			}
 		}
 	}
 }
@@ -196,6 +189,33 @@ ACompositeActor::ACompositeActor(const FObjectInitializer& ObjectInitializer)
 
 ACompositeActor::~ACompositeActor() = default;
 
+void ACompositeActor::PostLoad()
+{
+	Super::PostLoad();
+
+#if WITH_EDITORONLY_DATA
+	/**
+	* Refactored scene capture management where previous shadow/reflection maps were removed.
+	* Instead of polluting the actor with old data structures, we remove old components and rely on their recreation.
+	*/
+	if (GetLinkerCustomVersion(FFortniteMainBranchObjectVersion::GUID) < FFortniteMainBranchObjectVersion::CompositeActorSceneCaptureRefactor)
+	{
+		TArray<UActorComponent*> ActorComponents;
+		GetComponents(ActorComponents);
+
+		for (auto It = ActorComponents.CreateIterator(); It; ++It)
+		{
+			if (UCompositeShadowReflectionCatcherComponent* CaptureComponent = Cast<UCompositeShadowReflectionCatcherComponent>(*It))
+			{
+				RemoveInstanceComponent(CaptureComponent);
+				CaptureComponent->UnregisterComponent();
+				CaptureComponent->DestroyComponent();
+			}
+		}
+	}
+#endif // WITH_EDITORONLY_DATA
+}
+
 void ACompositeActor::PostRegisterAllComponents()
 {
 	using namespace UE::Composite::Private;
@@ -276,14 +296,12 @@ void ACompositeActor::PreEditChange(FProperty* PropertyAboutToChange)
 
 	Super::PreEditChange(PropertyAboutToChange);
 
+	PreEditCompositeLayers.Reset();
+
 	if (PropertyAboutToChange && PropertyAboutToChange->GetFName() == GET_MEMBER_NAME_CHECKED(ThisClass, CompositeLayers))
 	{
 		PreEditCompositeLayers = CompositeLayers;
 	}
-	else
-	{
-		PreEditCompositeLayers.Reset();
-	}
 }
 
 void ACompositeActor::PostEditChangeChainProperty(struct FPropertyChangedChainEvent& PropertyChangedChainEvent)
@@ -335,33 +353,17 @@ void ACompositeActor::PostEditChangeChainProperty(struct FPropertyChangedChainEv
 				// Always add elements at the start of the array (top of the stack)
 				Algo::Rotate(CompositeLayers, CompositeLayers.Num() - 1);
 			}
-			else if (PropertyChangedChainEvent.ChangeType == EPropertyChangeType::ValueSet || PropertyChangedChainEvent.ChangeType == EPropertyChangeType::Duplicate)
-			{
-				if (CompositeLayers.IsValidIndex(ArrayIndex))
-				{
-					if (const UCompositeLayerShadowReflection* CompositeLayer = Cast<UCompositeLayerShadowReflection>(CompositeLayers[ArrayIndex]))
-					{
-						FindOrCreateReflectionShadowCatcherPair(CompositeLayer, {});
-					}
-				}
-			}
-			else if (PropertyChangedChainEvent.ChangeType == EPropertyChangeType::ArrayRemove)
+
+			for (UCompositeLayerBase* PreEditCompositeLayer : PreEditCompositeLayers)
 			{
-				if (PreEditCompositeLayers.IsValidIndex(ArrayIndex))
+				if (IsValid(PreEditCompositeLayer))
 				{
-					if (const UCompositeLayerShadowReflection* CompositeLayer = Cast<UCompositeLayerShadowReflection>(PreEditCompositeLayers[ArrayIndex]))
+					if (!CompositeLayers.Contains(PreEditCompositeLayer))
 					{
-						DestroyReflectionShadowCatcherPair(CompositeLayer);
+						DestroySceneCaptures(PreEditCompositeLayer);
 					}
 				}
 			}
-			else if (PropertyChangedChainEvent.ChangeType == EPropertyChangeType::ArrayClear)
-			{
-				for (const UCompositeLayerBase* CompositeLayer : PreEditCompositeLayers)
-				{
-					DestroyReflectionShadowCatcherPair(CompositeLayer);
-				}
-			}
 
 			if (PropertyChangedChainEvent.PropertyChain.Num() <= 1)
 			{
@@ -440,71 +442,126 @@ void ACompositeActor::SetCamera(const FComponentReference& InComponentReference)
 	}
 }
 
-TShadowReflectionCatcherPair ACompositeActor::FindOrCreateReflectionShadowCatcherPair(const UCompositeLayerBase* InLayer, const TArray<TObjectPtr<AActor>>& InActorsToHide)
+template<>
+USceneCaptureComponent2D* ACompositeActor::FindOrCreateSceneCapture(const UCompositeLayerBase* InLayer, int32 InIndex, FName InBaseName)
 {
-	bool bComponentCreated = false;
+	if (!ensureMsgf(InLayer, TEXT("Expected non-null composite layer.")))
+	{
+		return nullptr;
+	}
 
-	if (!ReflectionShadowCatchers_CG.Contains(InLayer))
+	if (!ensure(InIndex >= 0))
 	{
-		const FName UniqueName = MakeUniqueObjectName(this, UCompositeShadowReflectionCatcherComponent::StaticClass(), FName("ShadowReflectionCatcher_CG"));
-		UCompositeShadowReflectionCatcherComponent* Component = NewObject<UCompositeShadowReflectionCatcherComponent>(this, UniqueName, RF_Transactional);
-		Component->SetupAttachment(RootComponent);
-		Component->RegisterComponent();
-		AddInstanceComponent(Component);
-		
-		ReflectionShadowCatchers_CG.Add(InLayer, Component);
-		bComponentCreated = true;
+		return nullptr;
 	}
 
-	if (!ReflectionShadowCatchers_NoCG.Contains(InLayer))
+	if (!SceneCapturesPerLayer.Contains(InLayer))
 	{
-		const FName UniqueName = MakeUniqueObjectName(this, UCompositeShadowReflectionCatcherComponent::StaticClass(), FName("ShadowReflectionCatcher_NoCG"));
-		UCompositeShadowReflectionCatcherComponent* Component = NewObject<UCompositeShadowReflectionCatcherComponent>(this, UniqueName, RF_Transactional);
-		Component->SetupAttachment(RootComponent);
-		Component->RegisterComponent();
-		AddInstanceComponent(Component);
+		SceneCapturesPerLayer.Add(InLayer, FSceneCaptureComponentArray{});
+	}
+
+	TArray<TObjectPtr<USceneCaptureComponent2D>>& SceneCaptures = SceneCapturesPerLayer[InLayer].Components;
 
-		ReflectionShadowCatchers_NoCG.Add(InLayer, Component);
-		bComponentCreated = true;
+	if (SceneCaptures.IsValidIndex(InIndex) && IsValid(SceneCaptures[InIndex]))
+	{
+		return SceneCaptures[InIndex];
 	}
 
+	const FName UniqueName = MakeUniqueObjectName(this, USceneCaptureComponent2D::StaticClass(), InBaseName);
+	USceneCaptureComponent2D* Component = NewObject<USceneCaptureComponent2D>(this, UniqueName, RF_Transactional);
+	Component->SetupAttachment(RootComponent);
+	Component->RegisterComponent();
+	Component->ComponentTags.Add(UE::Composite::DefaultCompositeTag);
+	AddInstanceComponent(Component);
+
+	UE::Composite::ConfigureSceneCapture(*Component, ESceneCapturePrimitiveRenderMode::PRM_UseShowOnlyList);
+	
+	TArray<FEngineShowFlagsSetting> ShowFlagSettings;
+	ShowFlagSettings.Add(FEngineShowFlagsSetting{ FEngineShowFlags::FindNameByIndex(FEngineShowFlags::EShowFlag::SF_Atmosphere), false });
+	ShowFlagSettings.Add(FEngineShowFlagsSetting{ FEngineShowFlags::FindNameByIndex(FEngineShowFlags::EShowFlag::SF_AllowPrimitiveAlphaHoldout), false });
+	ShowFlagSettings.Add(FEngineShowFlagsSetting{ FEngineShowFlags::FindNameByIndex(FEngineShowFlags::EShowFlag::SF_Bloom), false });
+	ShowFlagSettings.Add(FEngineShowFlagsSetting{ FEngineShowFlags::FindNameByIndex(FEngineShowFlags::EShowFlag::SF_EyeAdaptation), false });
+	ShowFlagSettings.Add(FEngineShowFlagsSetting{ FEngineShowFlags::FindNameByIndex(FEngineShowFlags::EShowFlag::SF_LocalExposure), false });
+	ShowFlagSettings.Add(FEngineShowFlagsSetting{ FEngineShowFlags::FindNameByIndex(FEngineShowFlags::EShowFlag::SF_PostProcessMaterial), false });
+	Component->SetShowFlagSettings(ShowFlagSettings);
+
+	SceneCaptures.EmplaceAt(InIndex, Component);
+
 #if WITH_EDITOR
-	if (bComponentCreated)
+	FModuleManager::LoadModuleChecked<FLevelEditorModule>("LevelEditor").BroadcastComponentsEdited();
+#endif
+
+	return Component;
+}
+
+template<>
+UCompositeShadowReflectionCatcherComponent* ACompositeActor::FindOrCreateSceneCapture(const UCompositeLayerBase* InLayer, int32 InIndex, FName InBaseName)
+{
+	if (!ensureMsgf(InLayer, TEXT("Expected non-null composite layer.")))
 	{
-		FModuleManager::LoadModuleChecked<FLevelEditorModule>("LevelEditor").BroadcastComponentsEdited();
+		return nullptr;
+	}
+
+	if (!ensure(InIndex >= 0))
+	{
+		return nullptr;
 	}
-#endif
 
-	TShadowReflectionCatcherPair OutPair;
-	OutPair[0] = ReflectionShadowCatchers_CG[InLayer];
-	OutPair[1] = ReflectionShadowCatchers_NoCG[InLayer];
-	OutPair[1]->HiddenActors = InActorsToHide;
+	if (!SceneCapturesPerLayer.Contains(InLayer))
+	{
+		SceneCapturesPerLayer.Add(InLayer, {});
+	}
+
+	TArray<TObjectPtr<USceneCaptureComponent2D>>& SceneCaptureArray = SceneCapturesPerLayer[InLayer].Components;
+
+	if (SceneCaptureArray.IsValidIndex(InIndex))
+	{
+		return Cast<UCompositeShadowReflectionCatcherComponent>(SceneCaptureArray[InIndex]);
+	}
 
-	return OutPair;
+	const FName UniqueName = MakeUniqueObjectName(this, UCompositeShadowReflectionCatcherComponent::StaticClass(), InBaseName);
+	UCompositeShadowReflectionCatcherComponent* Component = NewObject<UCompositeShadowReflectionCatcherComponent>(this, UniqueName, RF_Transactional);
+	Component->SetupAttachment(RootComponent);
+	Component->RegisterComponent();
+	Component->ComponentTags.Add(UE::Composite::DefaultCompositeTag);
+	AddInstanceComponent(Component);
+
+	SceneCaptureArray.EmplaceAt(InIndex, Component);
+
+#if WITH_EDITOR
+	FModuleManager::LoadModuleChecked<FLevelEditorModule>("LevelEditor").BroadcastComponentsEdited();
+#endif
+
+	return Component;
 }
 
-void ACompositeActor::DestroyReflectionShadowCatcherPair(const UCompositeLayerBase* InLayer)
+void ACompositeActor::ActivateSceneCaptures(const UCompositeLayerBase* InLayer, bool bInCaptureEveryFrame)
 {
-	bool bComponentRemoved = false;
-
+	if (FSceneCaptureComponentArray* SceneCapturesPtr = SceneCapturesPerLayer.Find(InLayer))
 	{
-		TObjectPtr<UCompositeShadowReflectionCatcherComponent> RemovedComponent;
-		if (ReflectionShadowCatchers_CG.RemoveAndCopyValue(InLayer, RemovedComponent))
+		for (TObjectPtr<USceneCaptureComponent2D>& SceneCapture : SceneCapturesPtr->Components)
 		{
-			if (IsValid(RemovedComponent))
+			if (IsValid(SceneCapture))
 			{
-				RemoveInstanceComponent(RemovedComponent);
-				RemovedComponent->UnregisterComponent();
-				RemovedComponent->DestroyComponent();
-
-				bComponentRemoved = true;
+				SceneCapture->bCaptureEveryFrame = bInCaptureEveryFrame;
 			}
 		}
 	}
+}
 
+void ACompositeActor::DestroySceneCaptures(const UCompositeLayerBase* InLayer)
+{
+	if (!ensureMsgf(InLayer, TEXT("Expected non-null composite layer.")))
 	{
-		TObjectPtr<UCompositeShadowReflectionCatcherComponent> RemovedComponent;
-		if (ReflectionShadowCatchers_NoCG.RemoveAndCopyValue(InLayer, RemovedComponent))
+		return;
+	}
+
+	bool bComponentRemoved = false;
+
+	FSceneCaptureComponentArray RemovedArray;
+	if (SceneCapturesPerLayer.RemoveAndCopyValue(InLayer, RemovedArray))
+	{
+		for (TObjectPtr<USceneCaptureComponent2D>& RemovedComponent : RemovedArray.Components)
 		{
 			if (IsValid(RemovedComponent))
 			{
@@ -525,21 +582,6 @@ void ACompositeActor::DestroyReflectionShadowCatcherPair(const UCompositeLayerBa
 #endif
 }
 
-void ACompositeActor::SetReflectionShadowCatcherCaptureEveryFrame(const UCompositeLayerBase* InLayer, bool bState)
-{
-	TObjectPtr<UCompositeShadowReflectionCatcherComponent>* Component0 = ReflectionShadowCatchers_CG.Find(InLayer);
-	if (Component0)
-	{
-		(*Component0)->bCaptureEveryFrame = bState;
-	}
-
-	TObjectPtr<UCompositeShadowReflectionCatcherComponent>* Component1 = ReflectionShadowCatchers_NoCG.Find(InLayer);
-	if (Component1)
-	{
-		(*Component1)->bCaptureEveryFrame = bState;
-	}
-}
-
 bool ACompositeActor::IsActive() const
 {
 	bool bIsActive = bIsEnabled;
@@ -560,8 +602,10 @@ void ACompositeActor::ApplyAllState()
 
 	OnEnabledStateChange.ExecuteIfBound(true, Subsystem);
 
-	UE::Composite::Private::SetCaptureEveryFrame(ReflectionShadowCatchers_CG, true);
-	UE::Composite::Private::SetCaptureEveryFrame(ReflectionShadowCatchers_NoCG, true);
+	for (const TObjectPtr<UCompositeLayerBase>& CompositeLayer : CompositeLayers)
+	{
+		ActivateSceneCaptures(CompositeLayer.Get(), true);
+	}
 }
 
 void ACompositeActor::RemoveAllState()
@@ -576,8 +620,10 @@ void ACompositeActor::RemoveAllState()
 		Subsystem->ResetBuiltInRenderPassOptions();
 	}
 
-	UE::Composite::Private::SetCaptureEveryFrame(ReflectionShadowCatchers_CG, false);
-	UE::Composite::Private::SetCaptureEveryFrame(ReflectionShadowCatchers_NoCG, false);
+	for (const TObjectPtr<UCompositeLayerBase>& CompositeLayer : CompositeLayers)
+	{
+		ActivateSceneCaptures(CompositeLayer.Get(), false);
+	}
 }
 
 
diff --git a/Engine/Plugins/Experimental/Compositing/Composite/Source/Composite/Private/CompositeUtilities.cpp b/Engine/Plugins/Experimental/Compositing/Composite/Source/Composite/Private/CompositeUtilities.cpp
new file mode 100644
--- /dev/null
+++ b/Engine/Plugins/Experimental/Compositing/Composite/Source/Composite/Private/CompositeUtilities.cpp
@@ -0,0 +1,31 @@
+// Copyright Epic Games, Inc. All Rights Reserved.
+
+#include "CompositeUtilities.h"
+
+void UE::Composite::ConfigureSceneCapture(USceneCaptureComponent2D& InComponent, const ESceneCapturePrimitiveRenderMode InPrimitiveRenderMode)
+{
+	InComponent.PrimitiveRenderMode = InPrimitiveRenderMode;
+	InComponent.CaptureSource = ESceneCaptureSource::SCS_FinalColorHDR;
+	InComponent.bUseRayTracingIfEnabled = true;
+	InComponent.bMainViewFamily = true;
+	InComponent.bMainViewResolution = true;
+	InComponent.bMainViewCamera = true;
+	// Inherit main view camera post-process settings to match main view DoF for example.
+	InComponent.bInheritMainViewCameraPostProcessSettings = true;
+	InComponent.bIgnoreScreenPercentage = true;
+	InComponent.bConsiderUnrenderedOpaquePixelAsFullyTranslucent = true;
+	InComponent.UnlitViewmode = ESceneCaptureUnlitViewmode::CaptureOrCustomRenderPass;
+
+	InComponent.PostProcessSettings.bOverride_ReflectionMethod = true;
+	InComponent.PostProcessSettings.ReflectionMethod = EReflectionMethod::Lumen;
+	InComponent.PostProcessSettings.bOverride_DynamicGlobalIlluminationMethod = true;
+	InComponent.PostProcessSettings.DynamicGlobalIlluminationMethod = EDynamicGlobalIlluminationMethod::Lumen;
+	InComponent.PostProcessSettings.bOverride_LumenFinalGatherScreenTraces = true;
+	InComponent.PostProcessSettings.LumenFinalGatherScreenTraces = false;
+	InComponent.PostProcessSettings.bOverride_LumenReflectionsScreenTraces = true;
+	InComponent.PostProcessSettings.LumenReflectionsScreenTraces = false;
+	// While hit lighting is more expensive, it significantly minimizes artifacts & noise. Here we default with higher-quality settings.
+	InComponent.PostProcessSettings.bOverride_LumenRayLightingMode = true;
+	InComponent.PostProcessSettings.LumenRayLightingMode = ELumenRayLightingModeOverride::HitLighting;
+}
+
diff --git a/Engine/Plugins/Experimental/Compositing/Composite/Source/Composite/Private/CompositeUtilities.h b/Engine/Plugins/Experimental/Compositing/Composite/Source/Composite/Private/CompositeUtilities.h
new file mode 100644
--- /dev/null
+++ b/Engine/Plugins/Experimental/Compositing/Composite/Source/Composite/Private/CompositeUtilities.h
@@ -0,0 +1,15 @@
+// Copyright Epic Games, Inc. All Rights Reserved.
+
+#pragma once
+
+#include "CoreTypes.h"
+#include "Components/SceneCaptureComponent2D.h"
+
+namespace UE
+{
+	namespace Composite
+	{
+		/** Configure scene captures for compositing (main view, raytracing, post-process settings, etc.) */
+		void ConfigureSceneCapture(USceneCaptureComponent2D& InComponent, const ESceneCapturePrimitiveRenderMode InPrimitiveRenderMode);
+	}
+}
diff --git a/Engine/Plugins/Experimental/Compositing/Composite/Source/Composite/Private/Layers/CompositeLayerSceneCapture.cpp b/Engine/Plugins/Experimental/Compositing/Composite/Source/Composite/Private/Layers/CompositeLayerSceneCapture.cpp
new file mode 100644
--- /dev/null
+++ b/Engine/Plugins/Experimental/Compositing/Composite/Source/Composite/Private/Layers/CompositeLayerSceneCapture.cpp
@@ -0,0 +1,199 @@
+// Copyright Epic Games, Inc. All Rights Reserved.
+
+#include "Layers/CompositeLayerSceneCapture.h"
+
+#include "CompositeActor.h"
+#include "Components/PrimitiveComponent.h"
+#include "Components/SceneCaptureComponent2D.h"
+#include "CompositeUtilities.h"
+#include "Engine/TextureRenderTarget2D.h"
+#include "UObject/Package.h"
+
+UCompositeLayerSceneCapture::UCompositeLayerSceneCapture(const FObjectInitializer& ObjectInitializer)
+	: Super(ObjectInitializer)
+{
+	Operation = ECompositeCoreMergeOp::Over;
+	bCustomRenderPass = false;
+}
+
+UCompositeLayerSceneCapture::~UCompositeLayerSceneCapture() = default;
+
+#if WITH_EDITOR
+void UCompositeLayerSceneCapture::PostEditChangeProperty(FPropertyChangedEvent& PropertyChangedEvent)
+{
+	Super::PostEditChangeProperty(PropertyChangedEvent);
+
+	const FName PropertyName = PropertyChangedEvent.GetMemberPropertyName();
+	
+	if (PropertyName == GET_MEMBER_NAME_CHECKED(ThisClass, bIsEnabled))
+	{
+		ACompositeActor* CompositeActor = GetTypedOuter<ACompositeActor>();
+		if (IsValid(CompositeActor))
+		{
+			CompositeActor->ActivateSceneCaptures(this, bIsEnabled);
+		}
+	}
+	else if (PropertyName == GET_MEMBER_NAME_CHECKED(ThisClass, Actors))
+	{
+		SetActors(MoveTemp(Actors));
+	}
+	else if (PropertyName == GET_MEMBER_NAME_CHECKED(ThisClass, bCustomRenderPass))
+	{
+		SetCustomRenderPass(bCustomRenderPass);
+	}
+	else if (PropertyName == GET_MEMBER_NAME_CHECKED(ThisClass, Postprocessing))
+	{
+		if (PropertyChangedEvent.ChangeType == EPropertyChangeType::ArrayAdd)
+		{
+			// Always add elements at the start of the array (top of the stack)
+			Algo::Rotate(Postprocessing, Postprocessing.Num() - 1);
+		}
+	}
+}
+#endif
+
+bool UCompositeLayerSceneCapture::GetProxy(FTraversalContext& InContext, FSceneRenderingBulkObjectAllocator& InFrameAllocator, FCompositeCorePassProxy*& OutProxy) const
+{
+	using namespace UE::CompositeCore;
+	constexpr int32 NumInputs = 2;
+
+	ACompositeActor* CompositeActor = GetTypedOuter<ACompositeActor>();
+	if (!IsValid(CompositeActor))
+	{
+		return false;
+	}
+
+	USceneCaptureComponent2D* CaptureComponent = GetSceneCaptureComponent();
+	if (!IsValid(CaptureComponent))
+	{
+		return false;
+	}
+
+	if (!IsValid(CaptureComponent->TextureTarget))
+	{
+		const FIntPoint& RenderResolution = CompositeActor->RenderResolution;
+
+		const FName RTName = MakeUniqueObjectName(GetTransientPackage(), UTextureRenderTarget2D::StaticClass(), TEXT("CompositeCaptureRT"));
+		CaptureComponent->TextureTarget = NewObject<UTextureRenderTarget2D>(GetTransientPackage(), RTName, RF_Transient);
+		CaptureComponent->TextureTarget->SizeX = RenderResolution.X;
+		CaptureComponent->TextureTarget->SizeY = RenderResolution.Y;
+		CaptureComponent->TextureTarget->RenderTargetFormat = ETextureRenderTargetFormat::RTF_RGBA16f;
+		CaptureComponent->TextureTarget->bForceLinearGamma = true;
+		//CaptureComponent->TextureTarget->bSupportsUAV = true;
+
+		CaptureComponent->TextureTarget->UpdateResource();
+	}
+
+	FExternalTexture& ExternalTexture = InContext.ExternalTextures.AddDefaulted_GetRef();
+	ExternalTexture.Texture = CaptureComponent->TextureTarget;
+	ExternalTexture.Metadata.bInvertedAlpha = true;
+	ExternalTexture.Metadata.bDistorted = false;
+	ExternalTexture.Metadata.Encoding = UE::CompositeCore::EEncoding::Linear;
+
+	FPassInputDeclArray Inputs;
+	Inputs.SetNum(NumInputs);
+
+	const int32 TexIndex = InContext.ExternalTextures.Num() - 1;
+	Inputs[0].Set<FPassExternalResourceDesc>({ EXTERNAL_RANGE_START_ID + TexIndex });
+	Inputs[1] = GetDefaultSecondInput(InContext);
+
+	// Copy for subpasses before we move the declared inputs
+	FPassInputDecl BaseInput0 = Inputs[0];
+
+	FMergePassProxy* Proxy = InFrameAllocator.Create<FMergePassProxy>(MoveTemp(Inputs));
+	Proxy->MergeOp = GetOperation(InContext);
+
+	if (!Postprocessing.IsEmpty())
+	{
+		TArray<const FCompositeCorePassProxy*> SubPassProxies;
+		SubPassProxies.Reserve(Postprocessing.Num());
+
+		bool bFirstValidPass = true;
+
+		// Iterate backwards so that lower passes in the UI are executed first
+		for (int32 PassIndex = Postprocessing.Num() - 1; PassIndex >= 0; --PassIndex)
+		{
+			const TObjectPtr<UCompositePassBase>& SubPass = Postprocessing[PassIndex];
+
+			if (IsValid(SubPass) && SubPass->IsActive())
+			{
+				FPassInputDecl PassInput;
+				if (bFirstValidPass)
+				{
+					PassInput = BaseInput0;
+					bFirstValidPass = false;
+				}
+				else
+				{
+					// Default input from previous pass
+					PassInput.Set<FPassInternalResourceDesc>({});
+				}
+
+				FCompositeCorePassProxy* SubPassProxy;
+				if (SubPass->GetProxy(PassInput, InFrameAllocator, SubPassProxy))
+				{
+					SubPassProxies.Add(SubPassProxy);
+				}
+			}
+		}
+
+		// Apply sub-passes overriding the first input
+		Proxy->AddSubPasses(0, SubPassProxies);
+	}
+
+	OutProxy = Proxy;
+	return true;
+}
+
+const TArray<AActor*> UCompositeLayerSceneCapture::GetActors() const
+{
+	return Actors;
+}
+
+void UCompositeLayerSceneCapture::SetActors(TArray<AActor*> InActors)
+{
+	USceneCaptureComponent2D* SceneCaptureComponent = GetSceneCaptureComponent();
+	if (IsValid(SceneCaptureComponent))
+	{
+		SceneCaptureComponent->ShowOnlyActors = InActors;
+	}
+
+	Actors = MoveTemp(InActors);
+}
+
+bool UCompositeLayerSceneCapture::IsCustomRenderPass() const
+{
+	return bCustomRenderPass;
+}
+
+void UCompositeLayerSceneCapture::SetCustomRenderPass(bool bInIsFastRenderPass)
+{
+	bCustomRenderPass = bInIsFastRenderPass;
+
+	USceneCaptureComponent2D* SceneCaptureComponent = GetSceneCaptureComponent();
+	if (IsValid(SceneCaptureComponent))
+	{
+		if (bCustomRenderPass)
+		{
+			SceneCaptureComponent->CaptureSource = ESceneCaptureSource::SCS_SceneColorHDR;
+			SceneCaptureComponent->bRenderInMainRenderer = true;
+		}
+		else
+		{
+			SceneCaptureComponent->CaptureSource = ESceneCaptureSource::SCS_FinalColorHDR;
+			SceneCaptureComponent->bRenderInMainRenderer = false;
+		}
+	}
+}
+
+USceneCaptureComponent2D* UCompositeLayerSceneCapture::GetSceneCaptureComponent() const
+{
+	ACompositeActor* CompositeActor = GetTypedOuter<ACompositeActor>();
+	if (IsValid(CompositeActor))
+	{
+		return CompositeActor->FindOrCreateSceneCapture<USceneCaptureComponent2D>(this);
+	}
+
+	return nullptr;
+}
+
diff --git a/Engine/Plugins/Experimental/Compositing/Composite/Source/Composite/Private/Layers/CompositeLayerShadowReflection.cpp b/Engine/Plugins/Experimental/Compositing/Composite/Source/Composite/Private/Layers/CompositeLayerShadowReflection.cpp
--- a/Engine/Plugins/Experimental/Compositing/Composite/Source/Composite/Private/Layers/CompositeLayerShadowReflection.cpp
+++ b/Engine/Plugins/Experimental/Compositing/Composite/Source/Composite/Private/Layers/CompositeLayerShadowReflection.cpp
@@ -27,7 +27,7 @@ void UCompositeLayerShadowReflection::PostEditChangeProperty(FPropertyChangedEve
 		ACompositeActor* CompositeActor = GetTypedOuter<ACompositeActor>();
 		if (IsValid(CompositeActor))
 		{
-			CompositeActor->SetReflectionShadowCatcherCaptureEveryFrame(this, bIsEnabled);
+			CompositeActor->ActivateSceneCaptures(this, bIsEnabled);
 		}
 	}
 	else if (PropertyName == GET_MEMBER_NAME_CHECKED(ThisClass, Actors))
@@ -51,19 +51,25 @@ bool UCompositeLayerShadowReflection::GetProxy(FTraversalContext& InContext, FSc
 		return false;
 	}
 
-	const TShadowReflectionCatcherPair ShadowReflectionPair = CompositeActor->FindOrCreateReflectionShadowCatcherPair(this, Actors);
+	TStaticArray<UCompositeShadowReflectionCatcherComponent*, 2> ShadowReflectionPair;
+	ShadowReflectionPair[0] = CompositeActor->FindOrCreateSceneCapture<UCompositeShadowReflectionCatcherComponent>(this, 0, FName("ShadowReflectionCatcher_CG"));
+	ShadowReflectionPair[1] = CompositeActor->FindOrCreateSceneCapture<UCompositeShadowReflectionCatcherComponent>(this, 1, FName("ShadowReflectionCatcher_NoCG"));
+
 	FMergePassProxy* Proxy = nullptr;
 	FMergePassProxy* DivisionProxy = nullptr;
 
 	bool bValidComponents = true;
 
-	for (const TObjectPtr<UCompositeShadowReflectionCatcherComponent>& Component : ShadowReflectionPair)
+	for (const UCompositeShadowReflectionCatcherComponent* Component : ShadowReflectionPair)
 	{
 		bValidComponents &= IsValid(Component);
 	}
 
 	if(bValidComponents)
 	{
+		// Update hidden actors
+		ShadowReflectionPair[1]->HiddenActors = Actors;
+
 		FPassInputDeclArray Inputs;
 		Inputs.SetNum(NumInputs);
 
diff --git a/Engine/Plugins/Experimental/Compositing/Composite/Source/Composite/Public/CompositeActor.h b/Engine/Plugins/Experimental/Compositing/Composite/Source/Composite/Public/CompositeActor.h
--- a/Engine/Plugins/Experimental/Compositing/Composite/Source/Composite/Public/CompositeActor.h
+++ b/Engine/Plugins/Experimental/Compositing/Composite/Source/Composite/Public/CompositeActor.h
@@ -16,11 +16,19 @@ class ACameraActor;
 class UCompositeCoreSubsystem;
 class UCompositeShadowReflectionCatcherComponent;
 class UPrimitiveComponent;
-
-using TShadowReflectionCatcherPair = TStaticArray<TObjectPtr<UCompositeShadowReflectionCatcherComponent>, 2>;
+class USceneCaptureComponent2D;
 
 DECLARE_DELEGATE_TwoParams(FOnCompositeActorEnabledStateChange, bool /* */, UCompositeCoreSubsystem* /* */);
 
+USTRUCT()
+struct FSceneCaptureComponentArray
+{
+	GENERATED_BODY()
+
+	UPROPERTY()
+	TArray<TObjectPtr<USceneCaptureComponent2D>> Components;
+};
+
 /** Actor used to control properties of the composite pipeline. */
 UCLASS(MinimalAPI)
 class ACompositeActor : public AActor
@@ -34,6 +42,10 @@ public:
 	/** Destructor. */
 	UE_API ~ACompositeActor();
 
+	//~ Begin UObject interface
+	UE_API virtual void PostLoad() override;
+	//~ End UObject interface
+
 	//~ Begin AActor Interface
 	UE_API virtual void PostRegisterAllComponents() override;
 	UE_API virtual void UnregisterAllComponents(bool bForReregister = false) override;
@@ -104,14 +116,15 @@ public:
 	int32 ViewUserFlags;
 
 public:
-	/** Convenience function used by layers to request a reflection & shadow catcher component pair. */
-	UE_API TShadowReflectionCatcherPair FindOrCreateReflectionShadowCatcherPair(const UCompositeLayerBase* InLayer, const TArray<TObjectPtr<AActor>>& InActorsToHide);
+	/** Convenience function used by layers to request a scene capture component. */
+	template<class RetType>
+	UE_API RetType* FindOrCreateSceneCapture(const UCompositeLayerBase* InLayer, int32 InIndex = 0, FName InBaseName = NAME_None);
 
-	/** Convenience function used by layers to remove and destroy a reflection & shadow catcher component pair. */
-	UE_API void DestroyReflectionShadowCatcherPair(const UCompositeLayerBase* InLayer);
+	/** Convenience function used by layers to control the capture state of scene captures. */
+	UE_API void ActivateSceneCaptures(const UCompositeLayerBase* InLayer, bool bInCaptureEveryFrame);
 
-	/** Convenience function used by layers to control the capture state of their reflection & shadow catcher component pair. */
-	UE_API void SetReflectionShadowCatcherCaptureEveryFrame(const UCompositeLayerBase* InLayer, bool bState);
+	/** Convenience function used by layers to remove and scene capture components. */
+	UE_API void DestroySceneCaptures(const UCompositeLayerBase* InLayer);
 
 	/** Returns true if the composite is active. */
 	UE_API bool IsActive() const;
@@ -123,13 +136,9 @@ private:
 	/** Remove all actor state from the composite subsystem. */
 	void RemoveAllState();
 
-	/** Managed reflection & shadow catcher components. */
-	UPROPERTY()
-	TMap<TWeakObjectPtr<const UCompositeLayerBase>, TObjectPtr<UCompositeShadowReflectionCatcherComponent>> ReflectionShadowCatchers_CG;
-
-	/** Managed reflection & shadow catcher components. */
+	/** Layer-managed scene capture components. */
 	UPROPERTY()
-	TMap<TWeakObjectPtr<const UCompositeLayerBase>, TObjectPtr<UCompositeShadowReflectionCatcherComponent>> ReflectionShadowCatchers_NoCG;
+	TMap<TWeakObjectPtr<const UCompositeLayerBase>, FSceneCaptureComponentArray> SceneCapturesPerLayer;
 
 #if WITH_EDITORONLY_DATA
 	/** Used to update components after we remove layers. */
diff --git a/Engine/Plugins/Experimental/Compositing/Composite/Source/Composite/Public/Layers/CompositeLayerSceneCapture.h b/Engine/Plugins/Experimental/Compositing/Composite/Source/Composite/Public/Layers/CompositeLayerSceneCapture.h
new file mode 100644
--- /dev/null
+++ b/Engine/Plugins/Experimental/Compositing/Composite/Source/Composite/Public/Layers/CompositeLayerSceneCapture.h
@@ -0,0 +1,72 @@
+// Copyright Epic Games, Inc. All Rights Reserved.
+
+#pragma once
+
+#include "CompositeLayerBase.h"
+
+#include "CompositeLayerSceneCapture.generated.h"
+
+#define UE_API COMPOSITE_API
+
+class AActor;
+
+UCLASS(MinimalAPI, BlueprintType, Blueprintable, EditInlineNew, CollapseCategories, meta = (DisplayName = "Scene Capture"))
+class UCompositeLayerSceneCapture : public UCompositeLayerBase
+{
+	GENERATED_BODY()
+
+public:
+	/** Constructor. */
+	UE_API UCompositeLayerSceneCapture(const FObjectInitializer& ObjectInitializer);
+	
+	/** Destructor. */
+	UE_API virtual ~UCompositeLayerSceneCapture();
+
+#if WITH_EDITOR
+	UE_API virtual void PostEditChangeProperty(FPropertyChangedEvent& PropertyChangedEvent) override;
+#endif
+
+	/** Getter function to override, returning pass layer proxies to be passed to the render thread. (Proxy objects should be allocated from the provided allocator.) */
+	UE_API virtual bool GetProxy(FTraversalContext& InContext, FSceneRenderingBulkObjectAllocator& InFrameAllocator, FCompositeCorePassProxy*& OutProxy) const override;
+
+	/** Get reflection & shadow caster actors. */
+	UFUNCTION(BlueprintGetter)
+	UE_API const TArray<AActor*> GetActors() const;
+
+	/** Set reflection & shadow caster actors. */
+	UFUNCTION(BlueprintSetter)
+	UE_API void SetActors(TArray<AActor*> InActors);
+
+	UFUNCTION(BlueprintGetter)
+	UE_API bool IsCustomRenderPass() const;
+
+	UFUNCTION(BlueprintSetter)
+	UE_API void SetCustomRenderPass(bool bInIsFastRenderPass);
+
+	/** Sub-passes applied during post-processing. */
+	UPROPERTY(EditAnywhere, BlueprintReadWrite, Instanced, Category = "", meta = (DisplayName = "Post-Processing", DisplayAfter = "bCustomRenderPass"))
+	TArray<TObjectPtr<UCompositePassBase>> Postprocessing;
+
+private:
+	/** Convenience function to access the actor-managed scene capture component. */
+	class USceneCaptureComponent2D* GetSceneCaptureComponent() const;
+
+	/**
+	 * List of actors to be rendered from the scene capture.
+	 * 
+	 * Automatically sets the scene capture's ShowOnlyActors list when updated.
+	*/
+	UPROPERTY(EditAnywhere, BlueprintGetter = GetActors, BlueprintSetter = SetActors, Category = "Composite", meta = (AllowPrivateAccess = true))
+	TArray<TObjectPtr<AActor>> Actors;
+
+	/**
+	 * Convenience toggle to render the scene capture as a fast/minimal Custom Render Pass, inlined in the main render & without support for lighting.
+	 * 
+	 * Automatically sets the scene capture's CaptureSource, bRenderInMainRenderer & bIgnoreScreenPercentage when updated.
+	 */
+	UPROPERTY(EditAnywhere, BlueprintGetter = IsCustomRenderPass, BlueprintSetter = SetCustomRenderPass, Category = "Composite", meta = (AllowPrivateAccess = true))
+	bool bCustomRenderPass;
+};
+
+#undef UE_API
+
diff --git a/Engine/Source/Runtime/Core/Public/UObject/FortniteMainBranchObjectVersions.inl b/Engine/Source/Runtime/Core/Public/UObject/FortniteMainBranchObjectVersions.inl
--- a/Engine/Source/Runtime/Core/Public/UObject/FortniteMainBranchObjectVersions.inl
+++ b/Engine/Source/Runtime/Core/Public/UObject/FortniteMainBranchObjectVersions.inl
@@ -664,5 +664,7 @@ FORTNITE_MAIN_VERSION(AddedMissingSerializationForPropertiesInDynamicMisc, 216)
 // Change default value for deprecated bEnableWorldPartitionGenerationSources
 FORTNITE_MAIN_VERSION(PCGDeprecateWorldPartitionGenerationSources, 217)
 
+// Refactored the composite (plugin) actor scene capture management.
+FORTNITE_MAIN_VERSION(CompositeActorSceneCaptureRefactor, 218)
 // -----<new versions can be added above this line>-------------------------------------------------
 #undef FORTNITE_MAIN_VERSION
\ No newline at end of file
