Composite: Add SMAA pass support.

diff --git a/Engine/Plugins/Experimental/Compositing/Composite/Source/Composite/Private/Passes/CompositePassSMAA.cpp b/Engine/Plugins/Experimental/Compositing/Composite/Source/Composite/Private/Passes/CompositePassSMAA.cpp
new file mode 100644
--- /dev/null
+++ b/Engine/Plugins/Experimental/Compositing/Composite/Source/Composite/Private/Passes/CompositePassSMAA.cpp
@@ -0,0 +1,99 @@
+// Copyright Epic Games, Inc. All Rights Reserved.
+
+#include "Passes/CompositePassSMAA.h"
+#include "Passes/CompositeCorePassFXAAProxy.h"
+
+#include "RenderGraphBuilder.h"
+#include "RenderGraphUtils.h"
+#include "SceneView.h"
+#include "ScreenPass.h"
+#include "PostProcess/SubpixelMorphologicalAA.h"
+
+DECLARE_GPU_STAT_NAMED(FCompositeCoreSMAA, TEXT("CompositeCore.SMAA"));
+
+namespace UE
+{
+	namespace Composite
+	{
+		using namespace CompositeCore;
+
+		class FSMAAPassProxy : public FCompositeCorePassProxy
+		{
+		public:
+			using FCompositeCorePassProxy::FCompositeCorePassProxy;
+
+			FPassOutput Add(FRDGBuilder& GraphBuilder, const FSceneView& InView, const FPassInputArray& Inputs, const FPassContext& PassContext) const override
+			{
+				RDG_EVENT_SCOPE_STAT(GraphBuilder, FCompositeCoreSMAA, "CompositeCore.SMAA");
+				RDG_GPU_STAT_SCOPE(GraphBuilder, FCompositeCoreSMAA);
+
+				check(ValidateInputs(Inputs));
+
+				FScreenPassTexture Input = Inputs[0].Texture;
+				FResourceMetadata Metadata = Inputs[0].Metadata;
+				const bool bLinearSourceColors = (Metadata.Encoding == EEncoding::Linear);
+
+				if (bLinearSourceColors)
+				{
+					// We tonemap & encode the result so that SMAA can operate on perceptual colors
+					constexpr bool bIsForward = true;
+					Input = CompositeCore::Private::AddDisplayTransformPass(GraphBuilder, InView, Input, {}, bIsForward);
+				}
+
+				static IConsoleVariable* SMAAQualityCVar = IConsoleManager::Get().FindConsoleVariable(TEXT("r.SMAA.Quality"));
+				check(SMAAQualityCVar);
+				const int32 Quality = QualityOverride.IsSet() ? QualityOverride.GetValue() : SMAAQualityCVar->GetInt();
+
+				FSMAAInputs SMAAInputs{};
+				SMAAInputs.SceneColor = Input;
+				SMAAInputs.SceneColorBeforeTonemap = Inputs[0].Texture;
+				SMAAInputs.Quality = static_cast<ESMAAQuality>(FMath::Clamp(Quality, 0, static_cast<int32>(ESMAAQuality::MAX) - 1));
+				if (!bLinearSourceColors)
+				{
+					SMAAInputs.OverrideOutput = Inputs.OverrideOutput;
+				}
+
+				FScreenPassTexture Output = AddSMAAPasses(GraphBuilder, InView, SMAAInputs);
+
+				if (bLinearSourceColors)
+				{
+					// We decode and invert the tonemapping to obtain linear colors again.
+					constexpr bool bIsForward = false;
+					Output = CompositeCore::Private::AddDisplayTransformPass(GraphBuilder, InView, Output, Inputs.OverrideOutput, bIsForward);
+				}
+
+				return FPassOutput{ MoveTemp(Output), Metadata, PassOutputOverride };
+			}
+
+			/** Optional r.SMAA.Quality setting override.  */
+			TOptional<int32> QualityOverride = {};
+		};
+	}
+}
+
+UCompositePassSMAA::UCompositePassSMAA(const FObjectInitializer& ObjectInitializer)
+	: Super(ObjectInitializer)
+	, Quality(5)
+{
+}
+
+UCompositePassSMAA::~UCompositePassSMAA() = default;
+
+bool UCompositePassSMAA::GetProxy(const UE::CompositeCore::FPassInputDecl& InputDecl, FSceneRenderingBulkObjectAllocator& InFrameAllocator, FCompositeCorePassProxy*& OutProxy) const
+{
+	using namespace UE::Composite;
+
+	UE_CALL_ONCE([&]{
+		if (ensure(GEngine))
+		{
+			GEngine->LoadSMAATextures();
+		}
+	});
+
+	FSMAAPassProxy* Proxy = InFrameAllocator.Create<FSMAAPassProxy>(UE::CompositeCore::FPassInputDeclArray{ InputDecl });
+	Proxy->QualityOverride = Quality;
+
+	OutProxy = Proxy;
+	return true;
+}
+
diff --git a/Engine/Plugins/Experimental/Compositing/Composite/Source/Composite/Public/Passes/CompositePassSMAA.h b/Engine/Plugins/Experimental/Compositing/Composite/Source/Composite/Public/Passes/CompositePassSMAA.h
new file mode 100644
--- /dev/null
+++ b/Engine/Plugins/Experimental/Compositing/Composite/Source/Composite/Public/Passes/CompositePassSMAA.h
@@ -0,0 +1,34 @@
+// Copyright Epic Games, Inc. All Rights Reserved.
+
+#pragma once
+
+#include "CompositePassBase.h"
+
+#include "CompositePassSMAA.generated.h"
+
+#define UE_API COMPOSITE_API
+
+/** SMAA composite pass, mostly useful for applying on the composite render pass for CG/Motion Graphics. */
+UCLASS(MinimalAPI, BlueprintType, Blueprintable, EditInlineNew, CollapseCategories, meta=(DisplayName="SMAA Pass"))
+class UCompositePassSMAA : public UCompositePassBase
+{
+	GENERATED_BODY()
+
+public:
+	/** Constructor */
+	UE_API UCompositePassSMAA(const FObjectInitializer& ObjectInitializer);
+	
+	/** Destructor */
+	UE_API ~UCompositePassSMAA();
+
+	UE_API virtual bool GetProxy(const UE::CompositeCore::FPassInputDecl& InputDecl, FSceneRenderingBulkObjectAllocator& InFrameAllocator, FCompositeCorePassProxy*& OutProxy) const override;
+
+public:
+
+	/** FXAA quality control, matching the "r.fxaa.quality" console variable values. */
+	UPROPERTY(EditAnywhere, BlueprintReadWrite, Category = "", meta=(ClampMin="0", ClampMax="5", UIMin = "0", UIMax = "5"))
+	int32 Quality;
+};
+
+#undef UE_API
+
diff --git a/Engine/Plugins/Experimental/Compositing/CompositeCore/Source/CompositeCore/Private/Passes/CompositeCorePassFXAAProxy.cpp b/Engine/Plugins/Experimental/Compositing/CompositeCore/Source/CompositeCore/Private/Passes/CompositeCorePassFXAAProxy.cpp
--- a/Engine/Plugins/Experimental/Compositing/CompositeCore/Source/CompositeCore/Private/Passes/CompositeCorePassFXAAProxy.cpp
+++ b/Engine/Plugins/Experimental/Compositing/CompositeCore/Source/CompositeCore/Private/Passes/CompositeCorePassFXAAProxy.cpp
@@ -35,7 +35,7 @@ namespace UE
 	{
 		namespace Private
 		{
-			FScreenPassTexture AddDisplayTransformPass(FRDGBuilder& GraphBuilder, const FScreenPassViewInfo ViewInfo, const FScreenPassTexture& Input, const FScreenPassRenderTarget& OverrideOutput, bool bIsForward, float Gamma = 2.2f)
+			FScreenPassTexture AddDisplayTransformPass(FRDGBuilder& GraphBuilder, const FScreenPassViewInfo ViewInfo, const FScreenPassTexture& Input, const FScreenPassRenderTarget& OverrideOutput, bool bIsForward, float InGamma)
 			{
 				RDG_EVENT_SCOPE_STAT(GraphBuilder, FCompositeCoreDisplayTransform, "CompositeCore.DisplayTransform");
 				RDG_GPU_STAT_SCOPE(GraphBuilder, FCompositeCoreDisplayTransform);
@@ -59,8 +59,8 @@ namespace UE
 				Parameters->InputTexture = Input.Texture;
 				Parameters->InputSampler = TStaticSamplerState<>::GetRHI();
 				Parameters->bIsForward = bIsForward;
-				Parameters->Gamma = Gamma;
-				Parameters->InvGamma = 1.0f / Gamma;
+				Parameters->Gamma = InGamma;
+				Parameters->InvGamma = 1.0f / InGamma;
 				Parameters->bIsForward = bIsForward;
 				Parameters->RenderTargets[0] = Output.GetRenderTargetBinding();
 
diff --git a/Engine/Plugins/Experimental/Compositing/CompositeCore/Source/CompositeCore/Public/Passes/CompositeCorePassFXAAProxy.h b/Engine/Plugins/Experimental/Compositing/CompositeCore/Source/CompositeCore/Public/Passes/CompositeCorePassFXAAProxy.h
--- a/Engine/Plugins/Experimental/Compositing/CompositeCore/Source/CompositeCore/Public/Passes/CompositeCorePassFXAAProxy.h
+++ b/Engine/Plugins/Experimental/Compositing/CompositeCore/Source/CompositeCore/Public/Passes/CompositeCorePassFXAAProxy.h
@@ -11,6 +11,11 @@ namespace UE
 {
 	namespace CompositeCore
 	{
+		namespace Private
+		{
+			UE_API FScreenPassTexture AddDisplayTransformPass(FRDGBuilder& GraphBuilder, const FScreenPassViewInfo ViewInfo, const FScreenPassTexture& Input, const FScreenPassRenderTarget& OverrideOutput, bool bIsForward, float InGamma = 2.2f);
+		}
+
 		class FFXAAPassProxy : public FCompositeCorePassProxy
 		{
 		public:
