Composite: Render target pool for scene capture components.

diff --git a/Engine/Plugins/Experimental/Compositing/Composite/Source/Composite/Private/Components/CompositeShadowReflectionCatcherComponent.cpp b/Engine/Plugins/Experimental/Compositing/Composite/Source/Composite/Private/Components/CompositeShadowReflectionCatcherComponent.cpp
--- a/Engine/Plugins/Experimental/Compositing/Composite/Source/Composite/Private/Components/CompositeShadowReflectionCatcherComponent.cpp
+++ b/Engine/Plugins/Experimental/Compositing/Composite/Source/Composite/Private/Components/CompositeShadowReflectionCatcherComponent.cpp
@@ -2,6 +2,7 @@
 
 #include "Components/CompositeShadowReflectionCatcherComponent.h"
 
+#include "CompositeRenderTargetPool.h"
 #include "CompositeUtilities.h"
 #include "Engine/TextureRenderTarget2D.h"
 #include "UObject/Package.h"
@@ -25,25 +26,11 @@ UCompositeShadowReflectionCatcherComponent::UCompositeShadowReflectionCatcherCom
 	ShowFlags.SetTemporalAA(true);
 }
 
-UTextureRenderTarget2D* UCompositeShadowReflectionCatcherComponent::FindOrCreateRenderTarget(int32 InSizeX, int32 InSizeY)
+void UCompositeShadowReflectionCatcherComponent::BeginDestroy()
 {
-	// TODO: Temporary solution, evaluate alternatives.
-	if (!IsValid(TextureTarget))
-	{
-		const FName RTName = MakeUniqueObjectName(GetTransientPackage(), UTextureRenderTarget2D::StaticClass(), TEXT("ShadowReflectionCatcherRT"));
-		TextureTarget = NewObject<UTextureRenderTarget2D>(GetTransientPackage(), RTName, RF_Transient);
-		TextureTarget->SizeX = InSizeX;
-		TextureTarget->SizeY = InSizeY;
-		TextureTarget->RenderTargetFormat = ETextureRenderTargetFormat::RTF_RGBA16f;
-		TextureTarget->bForceLinearGamma = true;
-		//TextureTarget->bSupportsUAV = true;
-
-		TextureTarget->UpdateResource();
-	}
-
-	TextureTarget->ResizeTarget(InSizeX, InSizeY);
-
-	return TextureTarget;
+	Super::BeginDestroy();
+
+	FCompositeRenderTargetPool::Get().ReleaseTarget(TextureTarget);
 }
 
 #if WITH_EDITOR
diff --git a/Engine/Plugins/Experimental/Compositing/Composite/Source/Composite/Private/CompositeModule.cpp b/Engine/Plugins/Experimental/Compositing/Composite/Source/Composite/Private/CompositeModule.cpp
--- a/Engine/Plugins/Experimental/Compositing/Composite/Source/Composite/Private/CompositeModule.cpp
+++ b/Engine/Plugins/Experimental/Compositing/Composite/Source/Composite/Private/CompositeModule.cpp
@@ -7,15 +7,25 @@
 #include "Modules/ModuleManager.h"
 #include "ShaderCore.h"
 
+#include "CompositeRenderTargetPool.h"
+
 DEFINE_LOG_CATEGORY(LogComposite);
 
 void FCompositeModule::StartupModule()
 {
 	AddShaderSourceDirectoryMapping(TEXT("/Plugin/Composite"), FPaths::Combine(IPluginManager::Get().FindPlugin(TEXT("Composite"))->GetBaseDir(), TEXT("Shaders")));
+
+	FCoreDelegates::OnEnginePreExit.AddRaw(this, &FCompositeModule::OnEnginePreExit);
 }
 
 void FCompositeModule::ShutdownModule()
 {
+	FCoreDelegates::OnEnginePreExit.RemoveAll(this);
+}
+
+void FCompositeModule::OnEnginePreExit()
+{
+	FCompositeRenderTargetPool::Get().Empty();
 }
 
 IMPLEMENT_MODULE( FCompositeModule, Composite )
diff --git a/Engine/Plugins/Experimental/Compositing/Composite/Source/Composite/Private/CompositeRenderTargetPool.cpp b/Engine/Plugins/Experimental/Compositing/Composite/Source/Composite/Private/CompositeRenderTargetPool.cpp
new file mode 100644
--- /dev/null
+++ b/Engine/Plugins/Experimental/Compositing/Composite/Source/Composite/Private/CompositeRenderTargetPool.cpp
@@ -0,0 +1,242 @@
+// Copyright Epic Games, Inc. All Rights Reserved.
+
+#include "CompositeRenderTargetPool.h"
+
+#include "CompositeModule.h"
+#include "Engine/TextureRenderTarget2D.h"
+#include "Logging/StructuredLog.h"
+#include "UObject/Package.h"
+
+static TAutoConsoleVariable<int32> CVarCompositePoolFramesUntilFlush(
+	TEXT("Composite.Pool.FramesUntilFlush"),
+	1,
+	TEXT("In editor, you can alter the target size and render format. ")
+	TEXT("This may leave render resources pooled, but never reclaimed. Auto-Flushing returns those resources. ")
+	TEXT("For values greater than zero, the pooling system will wait that number of frames before a target is considered 'stale'."));
+
+const FIntPoint FCompositeRenderTargetPool::DefaultSize = FIntPoint(1920, 1080);
+
+FCompositeRenderTargetPool& FCompositeRenderTargetPool::Get()
+{
+	static TUniquePtr<FCompositeRenderTargetPool> PoolInstance = MakeUnique<FCompositeRenderTargetPool>();
+
+	return *PoolInstance;
+}
+
+TObjectPtr<UTextureRenderTarget2D> FCompositeRenderTargetPool::AcquireTarget(UObject* InAssignee, const FIntPoint& InSize, ETextureRenderTargetFormat InFormat)
+{
+	TObjectPtr<UTextureRenderTarget2D> RenderTarget = nullptr;
+
+	FRenderTargetDesc Desc;
+	Desc.Size = InSize;
+	Desc.Format = InFormat;
+	if (TArray<FPooledRenderTarget>* MatchedTargetsPtr = PooledTargets.Find(Desc))
+	{
+		TArray<FPooledRenderTarget>& MatchedTargets = *MatchedTargetsPtr;
+
+		if (MatchedTargets.Num() > 0)
+		{
+			RenderTarget = MatchedTargets.Pop().Texture;
+		}
+
+		if (MatchedTargets.IsEmpty())
+		{
+			PooledTargets.Remove(Desc);
+		}
+	}
+	
+	if(!IsValid(RenderTarget))
+	{
+		const FName RTName = MakeUniqueObjectName(GetTransientPackage(), UTextureRenderTarget2D::StaticClass(), TEXT("CompositePooledRT"));
+		RenderTarget = NewObject<UTextureRenderTarget2D>(GetTransientPackage(), RTName, RF_Transient);
+		RenderTarget->SizeX = InSize.X;
+		RenderTarget->SizeY = InSize.Y;
+		RenderTarget->RenderTargetFormat = InFormat;
+		RenderTarget->bForceLinearGamma = true;
+		RenderTarget->bSupportsUAV = true;
+		RenderTarget->AddressX = TA_Clamp;
+		RenderTarget->AddressY = TA_Clamp;
+
+		RenderTarget->UpdateResource();
+	}
+
+	AssignedTargets.Add(RenderTarget, InAssignee);
+
+	return RenderTarget;
+}
+
+bool FCompositeRenderTargetPool::ConditionalAcquireTarget(UObject* InAssignee, TObjectPtr<UTextureRenderTarget2D>& InOutTarget, const FIntPoint& InSize, ETextureRenderTargetFormat InFormat)
+{
+	if (IsValid(InOutTarget)
+		&& InOutTarget->SizeX == InSize.X
+		&& InOutTarget->SizeY == InSize.Y
+		&& InOutTarget->RenderTargetFormat == InFormat)
+	{
+		// No update required
+		return false;
+	}
+
+	// Optionally return the existing render target in case it was resized.
+	ReleaseTarget(InOutTarget.Get());
+
+	// Update with found or created render target
+	InOutTarget = AcquireTarget(InAssignee, InSize, InFormat);
+
+	return true;
+}
+
+void FCompositeRenderTargetPool::ReleaseTarget(const TObjectPtr<UTextureRenderTarget2D>& InTarget)
+{
+	if (AssignedTargets.Remove(InTarget) > 0)
+	{
+		const FRenderTargetDesc Desc{ FIntPoint(InTarget->SizeX, InTarget->SizeY), InTarget->RenderTargetFormat };
+		PooledTargets.FindOrAdd(Desc).Push(InTarget);
+	}
+	else
+	{
+		UE_LOG(LogComposite, Warning, TEXT("Attempting to release a render target that doesn't belong to this pool."));
+	}
+}
+
+void FCompositeRenderTargetPool::EmptyPool(bool bInForceCollectGarbage)
+{
+	for (auto DescIt = PooledTargets.CreateIterator(); DescIt; ++DescIt)
+	{
+		for (auto TargetIt = DescIt->Value.CreateIterator(); TargetIt; ++TargetIt)
+		{
+			TObjectPtr<UTextureRenderTarget2D>& Target = TargetIt->Texture;
+			Target->ReleaseResource();
+			Target->ConditionalBeginDestroy();
+			Target = nullptr;
+
+			TargetIt.RemoveCurrent();
+		}
+	}
+
+	PooledTargets.Empty();
+
+	if (bInForceCollectGarbage)
+	{
+		CollectGarbage(GARBAGE_COLLECTION_KEEPFLAGS);
+	}
+}
+
+void FCompositeRenderTargetPool::Empty()
+{
+	constexpr bool bForceCollectGarbage = false;
+	EmptyPool(bForceCollectGarbage);
+
+	for (auto TargetIt = AssignedTargets.CreateIterator(); TargetIt; ++TargetIt)
+	{
+		TObjectPtr<UTextureRenderTarget2D>& Target = TargetIt->Key;
+		Target->ReleaseResource();
+		Target->ConditionalBeginDestroy();
+		Target = nullptr;
+	}
+
+	AssignedTargets.Empty();
+
+	CollectGarbage(GARBAGE_COLLECTION_KEEPFLAGS);
+}
+
+int32 FCompositeRenderTargetPool::Num() const
+{
+	return PooledTargets.Num() + AssignedTargets.Num();
+}
+
+void FCompositeRenderTargetPool::AddReferencedObjects(FReferenceCollector& Collector)
+{
+	for (auto DescIt = PooledTargets.CreateIterator(); DescIt; ++DescIt)
+	{
+		for (FPooledRenderTarget& PooledTarget : DescIt->Value)
+		{
+			Collector.AddReferencedObject(PooledTarget.Texture);
+		}
+	}
+
+	for (auto TargetIt = AssignedTargets.CreateIterator(); TargetIt; ++TargetIt)
+	{
+		Collector.AddReferencedObject(TargetIt->Key);
+	}
+}
+
+FString FCompositeRenderTargetPool::GetReferencerName() const
+{
+	return TEXT("FCompositeRenderTargetPool");
+}
+
+#if WITH_EDITOR
+void FCompositeRenderTargetPool::Tick(float /*DeltaSeconds*/)
+{
+	const int32 FramesUntilFlush = CVarCompositePoolFramesUntilFlush.GetValueOnGameThread();
+
+	// Since we can run in the editor, and could continuously alter the target's render size,
+	// we want to flush unused targets (out of fear that they'd never be used again) - targets 
+	// regularly used should still be claimed at this point
+	if (FramesUntilFlush > 0)
+	{
+		bool bHasRemovedTarget = false;
+
+		for (auto DescIt = PooledTargets.CreateIterator(); DescIt; ++DescIt)
+		{
+			for (auto TargetIt = DescIt->Value.CreateIterator(); TargetIt; ++TargetIt)
+			{
+				if (TargetIt->StaleFrameCount >= FramesUntilFlush)
+				{
+					TObjectPtr<UTextureRenderTarget2D>& Target = TargetIt->Texture;
+					UE_LOGFMT(LogComposite, Verbose, "Releasing stale pooled render target {Name}", Target->GetFName());
+
+					Target->ReleaseResource();
+					Target->ConditionalBeginDestroy();
+					Target = nullptr;
+
+					TargetIt.RemoveCurrent();
+
+					bHasRemovedTarget = true;
+				}
+			}
+
+			if (DescIt->Value.IsEmpty())
+			{
+				DescIt.RemoveCurrent();
+			}
+		}
+
+		if (bHasRemovedTarget)
+		{
+			CollectGarbage(GARBAGE_COLLECTION_KEEPFLAGS);
+		}
+	}
+
+	for (auto DescIt = PooledTargets.CreateIterator(); DescIt; ++DescIt)
+	{
+		for (FPooledRenderTarget& PooledTarget : DescIt->Value)
+		{
+			++PooledTarget.StaleFrameCount;
+		}
+	}
+
+	for (auto TargetIt = AssignedTargets.CreateIterator(); TargetIt; ++TargetIt)
+	{
+		const TWeakObjectPtr<UObject>& Assignee = TargetIt->Value;
+		if (!Assignee.IsValid())
+		{
+			TargetIt.RemoveCurrent();
+
+			UE_LOG(LogComposite, Display, TEXT("Releasing assigned render target with invalid assignee object - possible leak?"));
+		}
+	}
+}
+
+//------------------------------------------------------------------------------
+bool FCompositeRenderTargetPool::IsTickable() const
+{
+	return Num() > 0;
+}
+
+//------------------------------------------------------------------------------
+TStatId FCompositeRenderTargetPool::GetStatId() const
+{
+	RETURN_QUICK_DECLARE_CYCLE_STAT(FCompositeRenderTargetPool, STATGROUP_Tickables);
+}
+#endif
diff --git a/Engine/Plugins/Experimental/Compositing/Composite/Source/Composite/Private/CompositeRenderTargetPool.h b/Engine/Plugins/Experimental/Compositing/Composite/Source/Composite/Private/CompositeRenderTargetPool.h
new file mode 100644
--- /dev/null
+++ b/Engine/Plugins/Experimental/Compositing/Composite/Source/Composite/Private/CompositeRenderTargetPool.h
@@ -0,0 +1,125 @@
+// Copyright Epic Games, Inc. All Rights Reserved.
+
+#pragma once
+
+#include "CoreTypes.h"
+#include "UObject/GCObject.h"
+#include "Engine/TextureRenderTarget2D.h"
+
+#if WITH_EDITOR
+#include "TickableEditorObject.h"
+#endif
+
+class UObject;
+class UTextureRenderTarget2D;
+
+/**
+* While we favor RDG render targets whenever possible, scene captures rely on regular
+* UTextureRenderTarget objects which by definition need to live across multiple renders
+* (each with their own render dependency graphs).
+* 
+* This class provides a central location for transient render target management for regular scene captures.
+*/
+class FCompositeRenderTargetPool final : public FGCObject
+#if WITH_EDITOR
+	, public FTickableEditorObject
+#endif
+{
+public:
+	/** Default render target size/resolution of 1080p */
+	static const FIntPoint DefaultSize;
+
+	/** Static singleton pool getter. */
+	static FCompositeRenderTargetPool& Get();
+
+	/**
+	* Acquire new or recyled render target.
+	* 
+	* @param InAssignee Assignee object whose lifetime is to be paired with the transient render target. If the assignee is invalid, we automatically destroy the render target resource.
+	* @param InSize Render target texture size.
+	* @param InFormat Render target texture format.
+	* @return Render target object.
+	*/
+	TObjectPtr<UTextureRenderTarget2D> AcquireTarget(UObject* InAssignee, const FIntPoint& InSize, ETextureRenderTargetFormat InFormat = ETextureRenderTargetFormat::RTF_RGBA16f);
+
+	/**
+	* Conditionally acquire new or recyled render target, only if the size & format are mismatched with the current render target.
+	*
+	* @param InAssignee Assignee object whose lifetime is to be paired with the transient render target. If the assignee is invalid, we automatically destroy the render target resource.
+	* @param InOutTarget Render target object to update.
+	* @param InSize Render target texture size.
+	* @param InFormat Render target texture format.
+	* @return True if the render target object was updated.
+	*/
+	bool ConditionalAcquireTarget(UObject* InAssignee, TObjectPtr<UTextureRenderTarget2D>& InOutTarget, const FIntPoint& InSize, ETextureRenderTargetFormat InFormat = ETextureRenderTargetFormat::RTF_RGBA16f);
+
+	/**
+	* Release a previously acquired render target, and return it into the pool (for 'Composite.Pool.FramesUntilFlush' number of frames).
+	* @param InTarget Render target to release.
+	*/
+	void ReleaseTarget(const TObjectPtr<UTextureRenderTarget2D>& InTarget);
+
+	/** Destroy available pool render targets. */
+	void EmptyPool(bool bInForceCollectGarbage = true);
+
+	/** Destroy all render targets, both available & assigned. */
+	void Empty();
+
+	/** Number of managed render targets, both available & assigned. */
+	int32 Num() const;
+
+	//~ Begin FGCObject interface
+	virtual void AddReferencedObjects(FReferenceCollector& Collector) override;
+	virtual FString GetReferencerName() const override;
+	//~ End FGCObject interface
+
+		//~ Begin FTickableEditorObject interface
+#if WITH_EDITOR
+	virtual void Tick(float DeltaSeconds) override;
+	virtual bool IsTickable() const override;
+	virtual TStatId GetStatId() const override;
+#endif
+	//~ End FTickableEditorObject interface
+
+private:
+
+	/** Render target descriptor for finding matching available targets in the pool. */
+	struct FRenderTargetDesc
+	{
+		FIntPoint Size;
+
+		ETextureRenderTargetFormat Format;
+
+		bool operator==(const FRenderTargetDesc& Rhs) const
+		{
+			return Size == Rhs.Size && Format == Rhs.Format;
+		}
+
+		friend uint32 GetTypeHash(const FRenderTargetDesc& InTargetDesc)
+		{
+			return HashCombine(GetTypeHash(InTargetDesc.Size), GetTypeHash(InTargetDesc.Format));
+		}
+	};
+
+	/** Pool render target struct with a stale frame counter. */
+	struct FPooledRenderTarget
+	{
+		FPooledRenderTarget() = default;
+		
+		FPooledRenderTarget(TObjectPtr<UTextureRenderTarget2D> InTextureTarget)
+			: Texture(MoveTemp(InTextureTarget))
+		{ }
+
+		TObjectPtr<UTextureRenderTarget2D> Texture;
+#if WITH_EDITOR
+		int32 StaleFrameCount = 0;
+#endif 
+	};
+
+	/** Available pooled render targets. */
+	TMap<const FRenderTargetDesc, TArray<FPooledRenderTarget>> PooledTargets;
+
+	/** Assigned render targets (to be returned into the pool). */
+	TMap<TObjectPtr<UTextureRenderTarget2D>, TWeakObjectPtr<UObject>> AssignedTargets;
+};
+
diff --git a/Engine/Plugins/Experimental/Compositing/Composite/Source/Composite/Private/Layers/CompositeLayerBase.cpp b/Engine/Plugins/Experimental/Compositing/Composite/Source/Composite/Private/Layers/CompositeLayerBase.cpp
--- a/Engine/Plugins/Experimental/Compositing/Composite/Source/Composite/Private/Layers/CompositeLayerBase.cpp
+++ b/Engine/Plugins/Experimental/Compositing/Composite/Source/Composite/Private/Layers/CompositeLayerBase.cpp
@@ -3,6 +3,7 @@
 #include "Layers/CompositeLayerBase.h"
 
 #include "CompositeActor.h"
+#include "CompositeRenderTargetPool.h"
 
 UCompositeLayerBase::UCompositeLayerBase(const FObjectInitializer& ObjectInitializer)
 	: Super(ObjectInitializer)
@@ -56,6 +57,18 @@ void UCompositeLayerBase::PostEditChangeProperty(FPropertyChangedEvent& Property
 }
 #endif //WITH_EDITOR
 
+const FIntPoint UCompositeLayerBase::GetRenderResolution() const
+{
+	const ACompositeActor* CompositeActor = GetTypedOuter<ACompositeActor>();
+	
+	if (IsValid(CompositeActor))
+	{
+		return CompositeActor->RenderResolution;
+	}
+
+	return FCompositeRenderTargetPool::DefaultSize;
+}
+
 ECompositeCoreMergeOp UCompositeLayerBase::GetOperation(const FTraversalContext& InContext) const
 {
 	return InContext.bIsSolo ? ECompositeCoreMergeOp::None : Operation;
diff --git a/Engine/Plugins/Experimental/Compositing/Composite/Source/Composite/Private/Layers/CompositeLayerPlate.cpp b/Engine/Plugins/Experimental/Compositing/Composite/Source/Composite/Private/Layers/CompositeLayerPlate.cpp
--- a/Engine/Plugins/Experimental/Compositing/Composite/Source/Composite/Private/Layers/CompositeLayerPlate.cpp
+++ b/Engine/Plugins/Experimental/Compositing/Composite/Source/Composite/Private/Layers/CompositeLayerPlate.cpp
@@ -7,6 +7,7 @@
 #include "CompositeAssetUserData.h"
 #include "CompositeCoreSubsystem.h"
 #include "CompositeCoreSettings.h"
+#include "CompositeRenderTargetPool.h"
 #include "Engine/Texture2D.h"
 #include "MediaTexture.h"
 #include "Engine/TextureRenderTarget2D.h"
@@ -157,7 +158,7 @@ void UCompositeLayerPlate::PostLoad()
 	{
 		if (IsValid(Texture) && GetValidPreprocessingPassNum() > 0)
 		{
-			TextureProcessed = CreateTransientRenderTarget();
+			FCompositeRenderTargetPool::Get().ConditionalAcquireTarget(this, TextureProcessed, GetRenderResolution());
 		}
 
 		UpdateCompositeMeshes();
@@ -168,6 +169,8 @@ void UCompositeLayerPlate::BeginDestroy()
 {
 	Super::BeginDestroy();
 
+	FCompositeRenderTargetPool::Get().ReleaseTarget(TextureProcessed);
+
 #if WITH_EDITOR
 	// If we are not the class default object...
 	if (!HasAnyFlags(RF_ClassDefaultObject))
@@ -299,7 +302,7 @@ void UCompositeLayerPlate::PostEditChangeProperty(FPropertyChangedEvent& Propert
 		{
 			if (IsValid(Texture) && !IsValid(TextureProcessed))
 			{
-				TextureProcessed = CreateTransientRenderTarget();
+				FCompositeRenderTargetPool::Get().ConditionalAcquireTarget(this, TextureProcessed, GetRenderResolution());
 			}
 		}
 
@@ -592,41 +595,6 @@ void UCompositeLayerPlate::OnPostSaveWorld(UWorld* InWorld, FObjectPostSaveConte
 }
 #endif
 
-TObjectPtr<UTextureRenderTarget2D> UCompositeLayerPlate::CreateTransientRenderTarget()
-{
-	FIntPoint Resolution = FIntPoint(1920, 1080);
-		
-	if (ACompositeActor* CompositeActor = GetTypedOuter<ACompositeActor>())
-	{
-		Resolution = CompositeActor->RenderResolution;
-	}
-
-#if 0 // Note: We may wish to expose options to retain the original texture resolution
-	if (const UMediaTexture* MediaTexture = Cast<UMediaTexture>(Texture))
-	{
-		Resolution = FIntPoint(MediaTexture->GetWidth(), MediaTexture->GetHeight());
-	}
-	else if (const UTexture2D* Texture2D = Cast<UTexture2D>(Texture))
-	{
-		Resolution = FIntPoint(Texture2D->GetSizeX(), Texture2D->GetSizeY());
-	}
-#endif
-
-	const FName UniqueName = MakeUniqueObjectName(this, UTextureRenderTarget2D::StaticClass(), FName("CompositeProcessedPlate"));
-	TObjectPtr<UTextureRenderTarget2D> RenderTarget = NewObject<UTextureRenderTarget2D>(this, UniqueName, RF_Transient);
-	RenderTarget->SizeX = Resolution.X;
-	RenderTarget->SizeY = Resolution.Y;
-	RenderTarget->RenderTargetFormat = ETextureRenderTargetFormat::RTF_RGBA16f;
-	RenderTarget->bForceLinearGamma = true;
-	RenderTarget->bSupportsUAV = true;
-	RenderTarget->AddressX = TA_Clamp;
-	RenderTarget->AddressY = TA_Clamp;
-
-	RenderTarget->UpdateResource();
-
-	return RenderTarget;
-}
-
 void UCompositeLayerPlate::UpdateCompositeMeshes(bool bRegisterOnly) const
 {
 	for (AActor* MeshActor : CompositeMeshes)
diff --git a/Engine/Plugins/Experimental/Compositing/Composite/Source/Composite/Private/Layers/CompositeLayerSceneCapture.cpp b/Engine/Plugins/Experimental/Compositing/Composite/Source/Composite/Private/Layers/CompositeLayerSceneCapture.cpp
--- a/Engine/Plugins/Experimental/Compositing/Composite/Source/Composite/Private/Layers/CompositeLayerSceneCapture.cpp
+++ b/Engine/Plugins/Experimental/Compositing/Composite/Source/Composite/Private/Layers/CompositeLayerSceneCapture.cpp
@@ -5,6 +5,7 @@
 #include "CompositeActor.h"
 #include "Components/PrimitiveComponent.h"
 #include "Components/SceneCaptureComponent2D.h"
+#include "CompositeRenderTargetPool.h"
 #include "CompositeUtilities.h"
 #include "Engine/TextureRenderTarget2D.h"
 #include "UObject/Package.h"
@@ -18,6 +19,16 @@ UCompositeLayerSceneCapture::UCompositeLayerSceneCapture(const FObjectInitialize
 
 UCompositeLayerSceneCapture::~UCompositeLayerSceneCapture() = default;
 
+void UCompositeLayerSceneCapture::BeginDestroy()
+{
+	Super::BeginDestroy();
+
+	if (USceneCaptureComponent2D* CaptureComponent = GetSceneCaptureComponent())
+	{
+		FCompositeRenderTargetPool::Get().ReleaseTarget(CaptureComponent->TextureTarget);
+	}
+}
+	
 #if WITH_EDITOR
 void UCompositeLayerSceneCapture::PostEditChangeProperty(FPropertyChangedEvent& PropertyChangedEvent)
 {
@@ -68,20 +79,7 @@ bool UCompositeLayerSceneCapture::GetProxy(FTraversalContext& InContext, FSceneR
 		return false;
 	}
 
-	if (!IsValid(CaptureComponent->TextureTarget))
-	{
-		const FIntPoint& RenderResolution = CompositeActor->RenderResolution;
-
-		const FName RTName = MakeUniqueObjectName(GetTransientPackage(), UTextureRenderTarget2D::StaticClass(), TEXT("CompositeCaptureRT"));
-		CaptureComponent->TextureTarget = NewObject<UTextureRenderTarget2D>(GetTransientPackage(), RTName, RF_Transient);
-		CaptureComponent->TextureTarget->SizeX = RenderResolution.X;
-		CaptureComponent->TextureTarget->SizeY = RenderResolution.Y;
-		CaptureComponent->TextureTarget->RenderTargetFormat = ETextureRenderTargetFormat::RTF_RGBA16f;
-		CaptureComponent->TextureTarget->bForceLinearGamma = true;
-		//CaptureComponent->TextureTarget->bSupportsUAV = true;
-
-		CaptureComponent->TextureTarget->UpdateResource();
-	}
+	FCompositeRenderTargetPool::Get().ConditionalAcquireTarget(CaptureComponent, CaptureComponent->TextureTarget, GetRenderResolution());
 
 	FExternalTexture& ExternalTexture = InContext.ExternalTextures.AddDefaulted_GetRef();
 	ExternalTexture.Texture = CaptureComponent->TextureTarget;
diff --git a/Engine/Plugins/Experimental/Compositing/Composite/Source/Composite/Private/Layers/CompositeLayerShadowReflection.cpp b/Engine/Plugins/Experimental/Compositing/Composite/Source/Composite/Private/Layers/CompositeLayerShadowReflection.cpp
--- a/Engine/Plugins/Experimental/Compositing/Composite/Source/Composite/Private/Layers/CompositeLayerShadowReflection.cpp
+++ b/Engine/Plugins/Experimental/Compositing/Composite/Source/Composite/Private/Layers/CompositeLayerShadowReflection.cpp
@@ -5,6 +5,7 @@
 #include "CompositeActor.h"
 #include "Components/PrimitiveComponent.h"
 #include "Components/CompositeShadowReflectionCatcherComponent.h"
+#include "CompositeRenderTargetPool.h"
 
 UCompositeLayerShadowReflection::UCompositeLayerShadowReflection(const FObjectInitializer& ObjectInitializer)
 	: Super(ObjectInitializer)
@@ -75,10 +76,10 @@ bool UCompositeLayerShadowReflection::GetProxy(FTraversalContext& InContext, FSc
 		for (int32 Index = 0; Index < 2; ++Index)
 		{
 			const TObjectPtr<UCompositeShadowReflectionCatcherComponent>& Component = ShadowReflectionPair[Index];
-			const FIntPoint& RenderResolution = CompositeActor->RenderResolution;
+			FCompositeRenderTargetPool::Get().ConditionalAcquireTarget(Component, Component->TextureTarget, GetRenderResolution());
 
 			FExternalTexture& ExternalTexture = InContext.ExternalTextures.AddDefaulted_GetRef();
-			ExternalTexture.Texture = Component->FindOrCreateRenderTarget(RenderResolution.X, RenderResolution.Y);
+			ExternalTexture.Texture = Component->TextureTarget;
 			ExternalTexture.Metadata.bInvertedAlpha = true;
 			ExternalTexture.Metadata.bDistorted = false;
 			ExternalTexture.Metadata.Encoding = UE::CompositeCore::EEncoding::Linear;
diff --git a/Engine/Plugins/Experimental/Compositing/Composite/Source/Composite/Public/Components/CompositeShadowReflectionCatcherComponent.h b/Engine/Plugins/Experimental/Compositing/Composite/Source/Composite/Public/Components/CompositeShadowReflectionCatcherComponent.h
--- a/Engine/Plugins/Experimental/Compositing/Composite/Source/Composite/Public/Components/CompositeShadowReflectionCatcherComponent.h
+++ b/Engine/Plugins/Experimental/Compositing/Composite/Source/Composite/Public/Components/CompositeShadowReflectionCatcherComponent.h
@@ -16,12 +16,13 @@ class UCompositeShadowReflectionCatcherComponent : public USceneCaptureComponent
 
 public:
 
+	//~ Begin UObject Interface
+	UE_API virtual void BeginDestroy() override;
+	//~ End UObject Interface
+
 #if WITH_EDITOR
 	UE_API virtual void PostEditChangeProperty(struct FPropertyChangedEvent& PropertyChangedEvent) override;
 #endif //WITH_EDITOR
-
-	/** Find or create the active scene capture render target. */
-	UE_API UTextureRenderTarget2D* FindOrCreateRenderTarget(int32 InSizeX, int32 InSizeY);
 };
 
 #undef UE_API
diff --git a/Engine/Plugins/Experimental/Compositing/Composite/Source/Composite/Public/CompositeModule.h b/Engine/Plugins/Experimental/Compositing/Composite/Source/Composite/Public/CompositeModule.h
--- a/Engine/Plugins/Experimental/Compositing/Composite/Source/Composite/Public/CompositeModule.h
+++ b/Engine/Plugins/Experimental/Compositing/Composite/Source/Composite/Public/CompositeModule.h
@@ -17,5 +17,9 @@ public:
 	virtual void StartupModule() override;
 	virtual void ShutdownModule() override;
 	// End of IModuleInterface interface
+
+private:
+	// Engine pre-exit callback
+	void OnEnginePreExit();
 };
 
diff --git a/Engine/Plugins/Experimental/Compositing/Composite/Source/Composite/Public/Layers/CompositeLayerBase.h b/Engine/Plugins/Experimental/Compositing/Composite/Source/Composite/Public/Layers/CompositeLayerBase.h
--- a/Engine/Plugins/Experimental/Compositing/Composite/Source/Composite/Public/Layers/CompositeLayerBase.h
+++ b/Engine/Plugins/Experimental/Compositing/Composite/Source/Composite/Public/Layers/CompositeLayerBase.h
@@ -75,6 +75,9 @@ protected:
 	/** Fixed number of inputs used by layer merge passes. */
 	static constexpr int32 FixedNumLayerInputs = 2;
 
+	/** Convenience function that returns the composite actor's render resolution. */
+	const FIntPoint GetRenderResolution() const;
+
 	/** Convenience function that returns the merge operation depending on the traversal context. */
 	ECompositeCoreMergeOp GetOperation(const FTraversalContext& InContext) const;
 
diff --git a/Engine/Plugins/Experimental/Compositing/Composite/Source/Composite/Public/Layers/CompositeLayerPlate.h b/Engine/Plugins/Experimental/Compositing/Composite/Source/Composite/Public/Layers/CompositeLayerPlate.h
--- a/Engine/Plugins/Experimental/Compositing/Composite/Source/Composite/Public/Layers/CompositeLayerPlate.h
+++ b/Engine/Plugins/Experimental/Compositing/Composite/Source/Composite/Public/Layers/CompositeLayerPlate.h
@@ -127,9 +127,6 @@ private:
 	/** Get number of valid pre-processing passes. */
 	int32 GetValidPreprocessingPassNum() const;
 
-	/** Convenience function to create a transient render target at the actor render resolution. */
-	TObjectPtr<UTextureRenderTarget2D> CreateTransientRenderTarget();
-
 private:
 #if WITH_EDITOR
 	/** Pre-edit list of registered composite meshes. */
diff --git a/Engine/Plugins/Experimental/Compositing/Composite/Source/Composite/Public/Layers/CompositeLayerSceneCapture.h b/Engine/Plugins/Experimental/Compositing/Composite/Source/Composite/Public/Layers/CompositeLayerSceneCapture.h
--- a/Engine/Plugins/Experimental/Compositing/Composite/Source/Composite/Public/Layers/CompositeLayerSceneCapture.h
+++ b/Engine/Plugins/Experimental/Compositing/Composite/Source/Composite/Public/Layers/CompositeLayerSceneCapture.h
@@ -22,6 +22,10 @@ public:
 	/** Destructor. */
 	UE_API virtual ~UCompositeLayerSceneCapture();
 
+	//~ Begin UObject Interface
+	UE_API virtual void BeginDestroy() override;
+	//~ End UObject Interface
+
 #if WITH_EDITOR
 	UE_API virtual void PostEditChangeProperty(FPropertyChangedEvent& PropertyChangedEvent) override;
 #endif
