Composite: Prevent release warning on null objects.

diff --git a/Engine/Plugins/Experimental/Compositing/Composite/Source/Composite/Private/CompositeRenderTargetPool.cpp b/Engine/Plugins/Experimental/Compositing/Composite/Source/Composite/Private/CompositeRenderTargetPool.cpp
--- a/Engine/Plugins/Experimental/Compositing/Composite/Source/Composite/Private/CompositeRenderTargetPool.cpp
+++ b/Engine/Plugins/Experimental/Compositing/Composite/Source/Composite/Private/CompositeRenderTargetPool.cpp
@@ -87,6 +87,11 @@ bool FCompositeRenderTargetPool::ConditionalAcquireTarget(UObject* InAssignee, T
 
 void FCompositeRenderTargetPool::ReleaseTarget(const TObjectPtr<UTextureRenderTarget2D>& InTarget)
 {
+	if (!IsValid(InTarget))
+	{
+		return;
+	}
+
 	if (AssignedTargets.Remove(InTarget) > 0)
 	{
 		const FRenderTargetDesc Desc{ FIntPoint(InTarget->SizeX, InTarget->SizeY), InTarget->RenderTargetFormat };
