Composite: Disable option to auto-configure actors as hidden in scene capture on shadow & reflection layer.

* Enum now has 3 modes (None, Visible & Hidden). None will leave properties untouched & unmanaged.

diff --git a/Engine/Plugins/Experimental/Compositing/Composite/Source/Composite/Private/Layers/CompositeLayerShadowReflection.cpp b/Engine/Plugins/Experimental/Compositing/Composite/Source/Composite/Private/Layers/CompositeLayerShadowReflection.cpp
--- a/Engine/Plugins/Experimental/Compositing/Composite/Source/Composite/Private/Layers/CompositeLayerShadowReflection.cpp
+++ b/Engine/Plugins/Experimental/Compositing/Composite/Source/Composite/Private/Layers/CompositeLayerShadowReflection.cpp
@@ -11,7 +11,7 @@ UCompositeLayerShadowReflection::UCompositeLayerShadowReflection(const FObjectIn
 	: Super(ObjectInitializer)
 {
 	Operation = ECompositeCoreMergeOp::Multiply;
-	bAutoConfigureActors = true;
+	AutoConfigureActors = ECompositeHiddenInSceneCaptureConfiguration::None;
 }
 
 UCompositeLayerShadowReflection::~UCompositeLayerShadowReflection() = default;
@@ -33,10 +33,11 @@ void UCompositeLayerShadowReflection::PostEditChangeProperty(FPropertyChangedEve
 	}
 	else if (PropertyName == GET_MEMBER_NAME_CHECKED(ThisClass, Actors))
 	{
-		if (bAutoConfigureActors)
-		{
-			ConfigurePrimitivesForShadowReflectionCatching();
-		}
+		ConfigurePrimitivesForHiddenShadowReflectionCatching();
+	}
+	else if (PropertyName == GET_MEMBER_NAME_CHECKED(ThisClass, AutoConfigureActors))
+	{
+		ConfigurePrimitivesForHiddenShadowReflectionCatching();
 	}
 }
 #endif
@@ -114,23 +115,34 @@ void UCompositeLayerShadowReflection::SetActors(TArray<AActor*> InActors)
 {
 	Actors = MoveTemp(InActors);
 
-	if (bAutoConfigureActors)
-	{
-		ConfigurePrimitivesForShadowReflectionCatching();
-	}
+	ConfigurePrimitivesForHiddenShadowReflectionCatching();
 }
 
-void UCompositeLayerShadowReflection::ConfigurePrimitivesForShadowReflectionCatching()
+void UCompositeLayerShadowReflection::ConfigurePrimitivesForHiddenShadowReflectionCatching()
 {
+	bool bHideActorsInSceneCapture = false;
+
+	switch (AutoConfigureActors)
+	{
+	case ECompositeHiddenInSceneCaptureConfiguration::None:
+		return;
+	case ECompositeHiddenInSceneCaptureConfiguration::Hidden:
+		bHideActorsInSceneCapture = true;
+		break;
+	case ECompositeHiddenInSceneCaptureConfiguration::Visible:
+	default:
+		break;
+	}
+
 	for (AActor* Actor : Actors)
 	{
 		if (IsValid(Actor))
 		{
 			if (UPrimitiveComponent* PrimitiveComponent = Actor->GetComponentByClass<UPrimitiveComponent>())
 			{
-				PrimitiveComponent->SetHiddenInSceneCapture(true);
-				PrimitiveComponent->SetAffectIndirectLightingWhileHidden(true);
-				PrimitiveComponent->SetCastHiddenShadow(true);
+				PrimitiveComponent->SetHiddenInSceneCapture(bHideActorsInSceneCapture);
+				PrimitiveComponent->SetAffectIndirectLightingWhileHidden(bHideActorsInSceneCapture);
+				PrimitiveComponent->SetCastHiddenShadow(bHideActorsInSceneCapture);
 			}
 		}
 	}
diff --git a/Engine/Plugins/Experimental/Compositing/Composite/Source/Composite/Public/Layers/CompositeLayerShadowReflection.h b/Engine/Plugins/Experimental/Compositing/Composite/Source/Composite/Public/Layers/CompositeLayerShadowReflection.h
--- a/Engine/Plugins/Experimental/Compositing/Composite/Source/Composite/Public/Layers/CompositeLayerShadowReflection.h
+++ b/Engine/Plugins/Experimental/Compositing/Composite/Source/Composite/Public/Layers/CompositeLayerShadowReflection.h
@@ -12,6 +12,15 @@
 
 class AActor;
 
+/* Automatic primitive configuration mode for registered primitives. */
+UENUM(BlueprintType)
+enum class ECompositeHiddenInSceneCaptureConfiguration : uint8
+{
+	None UMETA(ToolTip = "No primitive properties are updated"),
+	Visible UMETA(ToolTip = "The following properties are set to false: bHiddenInSceneCapture, bAffectIndirectLightingWhileHidden & bCastHiddenShadow."),
+	Hidden UMETA(ToolTip = "The following properties are set to true: bHiddenInSceneCapture, bAffectIndirectLightingWhileHidden & bCastHiddenShadow."),
+};
+
 UCLASS(MinimalAPI, BlueprintType, Blueprintable, EditInlineNew, CollapseCategories, meta = (DisplayName = "Shadow Reflection Catcher"))
 class UCompositeLayerShadowReflection : public UCompositeLayerBase
 {
@@ -40,18 +49,18 @@ public:
 	UE_API void SetActors(TArray<AActor*> InActors);
 
 	/**
-	 * When enabled, we automatically toggle the following (primitive component) properties on actors when registered:
+	 * Depending on the configuration mode, we automatically enable or disable the following (primitive component) properties:
 	 * 
-	 * * bHiddenInSceneCapture: true
-	 * * bAffectIndirectLightingWhileHidden: true
-	 * * bCastHiddenShadow: true
+	 * * bHiddenInSceneCapture
+	 * * bAffectIndirectLightingWhileHidden
+	 * * bCastHiddenShadow
 	*/
-	UPROPERTY(EditAnywhere, BlueprintReadWrite, Category = "Composite", AdvancedDisplay, meta = (DisplayName = "Auto-Configure Actors"))
-	bool bAutoConfigureActors;
+	UPROPERTY(EditAnywhere, BlueprintReadWrite, Category = "Composite", AdvancedDisplay, meta = (DisplayAfter = "Actors", DisplayName = "Auto-Configure Actors"))
+	ECompositeHiddenInSceneCaptureConfiguration AutoConfigureActors;
 
 private:
 	/** Convenience function to configure actor primitives. */
-	void ConfigurePrimitivesForShadowReflectionCatching();
+	void ConfigurePrimitivesForHiddenShadowReflectionCatching();
 
 	/** List of reflection & shadow caster actors. */
 	UPROPERTY(EditAnywhere, BlueprintGetter = GetActors, BlueprintSetter = SetActors, Category = "Composite")
