Composite: Preliminary access to relevant properties from sequencer.

* Unfortunately FComponentReference (for the activa camera) doesn't work.
* Dedicated track implementation willl be required to expose layer properties.

diff --git a/Engine/Plugins/Experimental/Compositing/Composite/Source/Composite/Public/Components/CompositeViewProjectionComponent.h b/Engine/Plugins/Experimental/Compositing/Composite/Source/Composite/Public/Components/CompositeViewProjectionComponent.h
--- a/Engine/Plugins/Experimental/Compositing/Composite/Source/Composite/Public/Components/CompositeViewProjectionComponent.h
+++ b/Engine/Plugins/Experimental/Compositing/Composite/Source/Composite/Public/Components/CompositeViewProjectionComponent.h
@@ -14,7 +14,7 @@ class UMaterialParameterCollection;
 class UWorld;
 
 /** Component responsible for continuously updating the specified material parameter collection with a camera view projection matrix (to be used for texture projection in materials). */
-UCLASS(MinimalAPI, EditInlineNew, HideCategories = (Activation, Transform, Lighting, Rendering, Tags, Cooking, Physics, LOD, AssetUserData, Navigation), ClassGroup = Composite, editinlinenew, meta = (BlueprintSpawnableComponent, DisplayName = "Camera View Projection"))
+UCLASS(MinimalAPI, EditInlineNew, HideCategories = (Activation, Transform, Lighting, Rendering, Tags, Cooking, Physics, LOD, AssetUserData, Navigation), ClassGroup = Composite, meta = (BlueprintSpawnableComponent, DisplayName = "Camera View Projection"))
 class UCompositeViewProjectionComponent : public UActorComponent
 {
 	GENERATED_UCLASS_BODY()
@@ -39,15 +39,15 @@ public:
 
 public:
 	/** Whether or not the component activates the view-projection matrix Material Parameter Collection update. */
-	UPROPERTY(EditAnywhere, BlueprintReadWrite, Category = "CameraProjection")
+	UPROPERTY(EditAnywhere, BlueprintReadWrite, Interp, Category = "CameraProjection")
 	bool bIsEnabled = true;
 
 	/** The Material Parameter Collection in which the view-projection matrix should be stored. */
-	UPROPERTY(EditAnywhere, BlueprintReadWrite, Category = "CameraProjection")
+	UPROPERTY(EditAnywhere, BlueprintReadWrite, Interp, Category = "CameraProjection")
 	TObjectPtr<UMaterialParameterCollection> MaterialParameterCollection;
 
 	/** Parameter name of the first element of the transform in the Material Parameter Collection set above.  Requires space for 4 vectors. */
-	UPROPERTY(EditAnywhere, BlueprintReadWrite, Category = "CameraProjection")
+	UPROPERTY(EditAnywhere, BlueprintReadWrite, Interp, Category = "CameraProjection")
 	FName MatrixParameterName;
 
 public:
diff --git a/Engine/Plugins/Experimental/Compositing/Composite/Source/Composite/Public/CompositeActor.h b/Engine/Plugins/Experimental/Compositing/Composite/Source/Composite/Public/CompositeActor.h
--- a/Engine/Plugins/Experimental/Compositing/Composite/Source/Composite/Public/CompositeActor.h
+++ b/Engine/Plugins/Experimental/Compositing/Composite/Source/Composite/Public/CompositeActor.h
@@ -85,12 +85,12 @@ private:
 	TObjectPtr<class UCompositeViewProjectionComponent> ViewProjectionComponent;
 
 	/** Whether or not the composite behavior is active. */
-	UPROPERTY(EditAnywhere, BlueprintGetter = IsEnabled, BlueprintSetter = SetEnabled, Category = "Composite", meta = (AllowPrivateAccess = true, DisplayPriority="0"))
+	UPROPERTY(EditAnywhere, BlueprintGetter = IsEnabled, BlueprintSetter = SetEnabled, Interp, Category = "Composite", meta = (AllowPrivateAccess = true, DisplayPriority="0"))
 	bool bIsEnabled;
 
 public:
 	/** Composite render (output) resolution. Currently only used for transient scene capture render targets. */
-	UPROPERTY(EditAnywhere, BlueprintReadWrite, Category = "Composite")
+	UPROPERTY(EditAnywhere, BlueprintReadWrite, Interp, Category = "Composite")
 	FIntPoint RenderResolution;
 
 private:
@@ -112,7 +112,7 @@ public:
 	bool bOverridesViewUserFlags;
 
 	/** Custom user flags value used to alter materials in the composite render pass. Set to one by default such that branching can be used in Lit materials. */
-	UPROPERTY(EditAnywhere, BlueprintReadWrite, Category = "Composite", AdvancedDisplay, meta = (EditCondition = "bOverridesViewUserFlags"))
+	UPROPERTY(EditAnywhere, BlueprintReadWrite, Interp, Category = "Composite", AdvancedDisplay, meta = (EditCondition = "bOverridesViewUserFlags"))
 	int32 ViewUserFlags;
 
 public:
