Composite: Remove multiple actor constraint.

Preventing multiple composite actors from being enabled at the same time will be enforced differently.

diff --git a/Engine/Plugins/Experimental/Compositing/Composite/Source/CompositeEditor/Private/CompositeEditorModule.cpp b/Engine/Plugins/Experimental/Compositing/Composite/Source/CompositeEditor/Private/CompositeEditorModule.cpp
--- a/Engine/Plugins/Experimental/Compositing/Composite/Source/CompositeEditor/Private/CompositeEditorModule.cpp
+++ b/Engine/Plugins/Experimental/Compositing/Composite/Source/CompositeEditor/Private/CompositeEditorModule.cpp
@@ -57,13 +57,13 @@ void FCompositeEditorModule::OnLevelActorAdded(AActor* InActor)
 
 		if (CompositeActorCount > 1)
 		{
-			UE_LOG(LogCompositeEditor, Error, TEXT("No more than one composite actor is currently allowed."));
+			ACompositeActor* CompositeActor = Cast<ACompositeActor>(InActor);
+			if (IsValid(CompositeActor))
+			{
+				CompositeActor->SetEnabled(false);
 
-			FNotificationInfo NotifyInfo(LOCTEXT("CompositeActorLimitMessage", "No more than one composite actor is currently allowed."));
-			NotifyInfo.ExpireDuration = 4.0f;
-			FSlateNotificationManager::Get().AddNotification(NotifyInfo)->SetCompletionState(SNotificationItem::CS_Fail);
-
-			InActor->Destroy();
+				UE_LOG(LogCompositeEditor, Warning, TEXT("No more than one composite actor should be enabled at the same time, disabling."));
+			}
 			break;
 		}
 	}
