Composite: Prevent RT deallocation crash on PIE exit.

* We release render targets before (PIE or any) duplication, such that their allocation always comes from the pool. This is necessary since scene capture render target properties are not instanced.
* Added very verbose logging to the pool for improved debugging.

diff --git a/Engine/Plugins/Experimental/Compositing/Composite/Source/Composite/Private/CompositeActor.cpp b/Engine/Plugins/Experimental/Compositing/Composite/Source/Composite/Private/CompositeActor.cpp
--- a/Engine/Plugins/Experimental/Compositing/Composite/Source/Composite/Private/CompositeActor.cpp
+++ b/Engine/Plugins/Experimental/Compositing/Composite/Source/Composite/Private/CompositeActor.cpp
@@ -6,6 +6,7 @@
 #include "CompositeModule.h"
 #include "CompositeCoreSettings.h"
 #include "CompositeUtilities.h"
+#include "CompositeRenderTargetPool.h"
 
 #include "Components/CompositeShadowReflectionCatcherComponent.h"
 #include "Components/CompositeViewProjectionComponent.h"
@@ -308,6 +309,29 @@ void ACompositeActor::Tick(float DeltaTime)
 }
 
 #if WITH_EDITOR
+void ACompositeActor::PreDuplicateFromRoot(FObjectDuplicationParameters& DupParams)
+{
+	Super::PreDuplicateFromRoot(DupParams);
+
+	/**
+	 * Whether DupParams.DuplicateMode == EDuplicateMode::PIE or otherwise,
+	 * we release transient scene capture render targets (which are not instanced properties)
+	 * to prevent their duplication. Allocation must instead come from the pool.
+	 * 
+	 * This is necessary since scene capture's TextureTarget property is not instanced.
+	 */
+	TInlineComponentArray<USceneCaptureComponent2D*> SceneCaptureComponents(this);
+	for (USceneCaptureComponent2D* CaptureComponent : SceneCaptureComponents)
+	{
+		FCompositeRenderTargetPool& Pool = FCompositeRenderTargetPool::Get();
+
+		if (Pool.IsAssignedTarget(CaptureComponent->TextureTarget))
+		{
+			Pool.ReleaseTarget(CaptureComponent->TextureTarget);
+		}
+	}
+}
+
 void ACompositeActor::PreEditChange(FProperty* PropertyAboutToChange)
 {
 	TOptional<TGuardValue<bool>> ModifyingPropertyGuard;
diff --git a/Engine/Plugins/Experimental/Compositing/Composite/Source/Composite/Private/CompositeRenderTargetPool.cpp b/Engine/Plugins/Experimental/Compositing/Composite/Source/Composite/Private/CompositeRenderTargetPool.cpp
--- a/Engine/Plugins/Experimental/Compositing/Composite/Source/Composite/Private/CompositeRenderTargetPool.cpp
+++ b/Engine/Plugins/Experimental/Compositing/Composite/Source/Composite/Private/CompositeRenderTargetPool.cpp
@@ -37,6 +37,8 @@ TObjectPtr<UTextureRenderTarget2D> FCompositeRenderTargetPool::AcquireTarget(UOb
 		if (MatchedTargets.Num() > 0)
 		{
 			RenderTarget = MatchedTargets.Pop().Texture;
+
+			UE_LOGFMT(LogComposite, VeryVerbose, "Acquiring pooled render target: {Name}", RenderTarget->GetFName());
 		}
 
 		if (MatchedTargets.IsEmpty())
@@ -58,6 +60,8 @@ TObjectPtr<UTextureRenderTarget2D> FCompositeRenderTargetPool::AcquireTarget(UOb
 		RenderTarget->AddressY = TA_Clamp;
 
 		RenderTarget->UpdateResource();
+
+		UE_LOGFMT(LogComposite, VeryVerbose, "Creating render target: {Name}", RenderTarget->GetFName());
 	}
 
 	AssignedTargets.Add(RenderTarget, InAssignee);
@@ -77,7 +81,7 @@ bool FCompositeRenderTargetPool::ConditionalAcquireTarget(UObject* InAssignee, T
 	}
 
 	// Optionally return the existing render target in case it was resized.
-	ReleaseTarget(InOutTarget.Get());
+	ReleaseTarget(InOutTarget);
 
 	// Update with found or created render target
 	InOutTarget = AcquireTarget(InAssignee, InSize, InFormat);
@@ -85,17 +89,26 @@ bool FCompositeRenderTargetPool::ConditionalAcquireTarget(UObject* InAssignee, T
 	return true;
 }
 
-void FCompositeRenderTargetPool::ReleaseTarget(const TObjectPtr<UTextureRenderTarget2D>& InTarget)
+bool FCompositeRenderTargetPool::IsAssignedTarget(const TObjectPtr<UTextureRenderTarget2D>& InTarget) const
+{
+	return AssignedTargets.Contains(InTarget);
+}
+
+void FCompositeRenderTargetPool::ReleaseTarget(TObjectPtr<UTextureRenderTarget2D>& InTarget)
 {
 	if (!IsValid(InTarget))
 	{
 		return;
 	}
 
+	UE_LOGFMT(LogComposite, VeryVerbose, "Releasing render target: {Name}", InTarget->GetFName());
+
 	if (AssignedTargets.Remove(InTarget) > 0)
 	{
 		const FRenderTargetDesc Desc{ FIntPoint(InTarget->SizeX, InTarget->SizeY), InTarget->RenderTargetFormat };
 		PooledTargets.FindOrAdd(Desc).Push(InTarget);
+
+		InTarget = nullptr;
 	}
 	else
 	{
@@ -105,6 +118,8 @@ void FCompositeRenderTargetPool::ReleaseTarget(const TObjectPtr<UTextureRenderTa
 
 void FCompositeRenderTargetPool::EmptyPool(bool bInForceCollectGarbage)
 {
+	UE_LOG(LogComposite, VeryVerbose, TEXT("Empty pool."));
+
 	for (auto DescIt = PooledTargets.CreateIterator(); DescIt; ++DescIt)
 	{
 		for (auto TargetIt = DescIt->Value.CreateIterator(); TargetIt; ++TargetIt)
@@ -131,6 +146,8 @@ void FCompositeRenderTargetPool::Empty()
 	constexpr bool bForceCollectGarbage = false;
 	EmptyPool(bForceCollectGarbage);
 
+	UE_LOG(LogComposite, VeryVerbose, TEXT("Empty assigned render target."));
+
 	for (auto TargetIt = AssignedTargets.CreateIterator(); TargetIt; ++TargetIt)
 	{
 		TObjectPtr<UTextureRenderTarget2D>& Target = TargetIt->Key;
diff --git a/Engine/Plugins/Experimental/Compositing/Composite/Source/Composite/Private/CompositeRenderTargetPool.h b/Engine/Plugins/Experimental/Compositing/Composite/Source/Composite/Private/CompositeRenderTargetPool.h
--- a/Engine/Plugins/Experimental/Compositing/Composite/Source/Composite/Private/CompositeRenderTargetPool.h
+++ b/Engine/Plugins/Experimental/Compositing/Composite/Source/Composite/Private/CompositeRenderTargetPool.h
@@ -53,11 +53,14 @@ public:
 	*/
 	bool ConditionalAcquireTarget(UObject* InAssignee, TObjectPtr<UTextureRenderTarget2D>& InOutTarget, const FIntPoint& InSize, ETextureRenderTargetFormat InFormat = ETextureRenderTargetFormat::RTF_RGBA16f);
 
+	/** Check if the render target was created by the pool and is currently assigned. */
+	bool IsAssignedTarget(const TObjectPtr<UTextureRenderTarget2D>& InTarget) const;
+
 	/**
 	* Release a previously acquired render target, and return it into the pool (for 'Composite.Pool.FramesUntilFlush' number of frames).
 	* @param InTarget Render target to release.
 	*/
-	void ReleaseTarget(const TObjectPtr<UTextureRenderTarget2D>& InTarget);
+	void ReleaseTarget(TObjectPtr<UTextureRenderTarget2D>& InTarget);
 
 	/** Destroy available pool render targets. */
 	void EmptyPool(bool bInForceCollectGarbage = true);
diff --git a/Engine/Plugins/Experimental/Compositing/Composite/Source/Composite/Public/CompositeActor.h b/Engine/Plugins/Experimental/Compositing/Composite/Source/Composite/Public/CompositeActor.h
--- a/Engine/Plugins/Experimental/Compositing/Composite/Source/Composite/Public/CompositeActor.h
+++ b/Engine/Plugins/Experimental/Compositing/Composite/Source/Composite/Public/CompositeActor.h
@@ -54,6 +54,7 @@ public:
 	//~ End AActor Interface
 
 #if WITH_EDITOR
+	UE_API virtual void PreDuplicateFromRoot(FObjectDuplicationParameters& DupParams) override;
 	UE_API virtual void PreEditChange(FProperty* PropertyAboutToChange) override;
 	UE_API virtual void PostEditChangeChainProperty(struct FPropertyChangedChainEvent& PropertyChangedChainEvent) override;
 	UE_API virtual void SetIsTemporarilyHiddenInEditor(bool bIsHidden) override;
