Composite: Add support for OCIO between layers via new special type of "processing-only" layer.

* Required apply inverse/forward transforms via OCIO.
* Exception to layer merging two inputs: has a single input since it's only a container for passes.
* We improve all layer descriptions in passing.

diff --git a/Engine/Plugins/Experimental/Compositing/Composite/Source/Composite/Private/CompositeActor.cpp b/Engine/Plugins/Experimental/Compositing/Composite/Source/Composite/Private/CompositeActor.cpp
--- a/Engine/Plugins/Experimental/Compositing/Composite/Source/Composite/Private/CompositeActor.cpp
+++ b/Engine/Plugins/Experimental/Compositing/Composite/Source/Composite/Private/CompositeActor.cpp
@@ -154,13 +154,17 @@ namespace UE
 					RegisterLayers(Context, Layers, *RenderWork.FrameAllocator, BranchedLayerProxies);
 
 					// Trim all proxies before our branch index (for the tonemap location) to instead recycle previous work.
-					BranchedLayerProxies.RemoveAt(0, BranchLayerIndex);
+					if (BranchedLayerProxies.IsValidIndex(BranchLayerIndex))
+					{
+						BranchedLayerProxies.RemoveAt(0, BranchLayerIndex);
+					}
 
 					if (ensureMsgf(!BranchedLayerProxies.IsEmpty(), TEXT("Unexpected empty branch proxy array.")) && LayerProxies.IsValidIndex(BranchSourceIndex))
 					{
 						// This is where we connect to the previous SSR input work branch.
 						// Input0 will be merged on top of Input1, the previous proxy's output.
-						BranchedLayerProxies[0]->GetDeclaredInput(1).Set<const FCompositeCorePassProxy*>(LayerProxies[BranchSourceIndex]);
+						const int32 NumInputs = BranchedLayerProxies[0]->GetNumDeclaredInputs();
+						BranchedLayerProxies[0]->GetDeclaredInput(NumInputs-1).Set<const FCompositeCorePassProxy*>(LayerProxies[BranchSourceIndex]);
 					}
 
 					RenderWork.FramePasses.Emplace(ISceneViewExtension::EPostProcessingPass::SSRInput, LayerProxies);
@@ -174,6 +178,18 @@ namespace UE
 				RenderWork.ExternalInputs = Context.ExternalTextures;
 				RenderWork.PreprocessingPasses = Context.PreprocessingPasses;
 
+				switch (InCompositeActor.OverrideMainRenderMode)
+				{
+				case ECompositeOverrideMainRenderMode::FinalColorHDR:
+					RenderWork.MainRenderMode = ESceneCaptureSource::SCS_FinalColorHDR;
+					break;
+				case ECompositeOverrideMainRenderMode::FinalToneCurveHDR:
+					RenderWork.MainRenderMode = ESceneCaptureSource::SCS_FinalToneCurveHDR;
+					break;
+				default:
+					break;
+				}
+
 				return RenderWork;
 			}
 		}
diff --git a/Engine/Plugins/Experimental/Compositing/Composite/Source/Composite/Private/Layers/CompositeLayerProcessing.cpp b/Engine/Plugins/Experimental/Compositing/Composite/Source/Composite/Private/Layers/CompositeLayerProcessing.cpp
new file mode 100644
--- /dev/null
+++ b/Engine/Plugins/Experimental/Compositing/Composite/Source/Composite/Private/Layers/CompositeLayerProcessing.cpp
@@ -0,0 +1,76 @@
+// Copyright Epic Games, Inc. All Rights Reserved.
+
+#include "Layers/CompositeLayerProcessing.h"
+
+#include "CompositeActor.h"
+
+UCompositeLayerProcessing::UCompositeLayerProcessing(const FObjectInitializer& ObjectInitializer)
+	: Super(ObjectInitializer)
+{
+	Operation = ECompositeCoreMergeOp::None;
+}
+
+UCompositeLayerProcessing::~UCompositeLayerProcessing() = default;
+
+#if WITH_EDITOR
+bool UCompositeLayerProcessing::CanEditChange(const FProperty* InProperty) const
+{
+	bool bIsEditable = Super::CanEditChange(InProperty);
+
+	if (bIsEditable && InProperty)
+	{
+		const FName PropertyName = InProperty->GetFName();
+
+		if (PropertyName == GET_MEMBER_NAME_CHECKED(ThisClass, bIsSolo))
+		{
+			bIsEditable = false;
+		}
+		else if (PropertyName == GET_MEMBER_NAME_CHECKED(ThisClass, Operation))
+		{
+			bIsEditable = false;
+		}
+	}
+
+	return bIsEditable;
+}
+
+void UCompositeLayerProcessing::PostEditChangeProperty(FPropertyChangedEvent& PropertyChangedEvent)
+{
+	Super::PostEditChangeProperty(PropertyChangedEvent);
+}
+#endif
+
+bool UCompositeLayerProcessing::GetProxy(FTraversalContext& InContext, FSceneRenderingBulkObjectAllocator& InFrameAllocator, FCompositeCorePassProxy*& OutProxy) const
+{
+	using namespace UE::CompositeCore;
+
+	FPassInputDecl BasePassInput = GetDefaultSecondInput(InContext);
+	OutProxy = nullptr;
+
+	// Iterate backwards so that lower passes in the UI are executed first
+	// TODO: Flip back once we have UI customization displaying passes in the correct bottom-up order
+	for (int32 PassIndex = LayerPasses.Num() - 1; PassIndex >= 0; --PassIndex)
+	{
+		const TObjectPtr<UCompositePassBase>& Pass = LayerPasses[PassIndex];
+
+		if (IsValid(Pass) && Pass->IsActive())
+		{
+			FCompositeCorePassProxy* PassProxy;
+			if (Pass->GetProxy(BasePassInput, InFrameAllocator, PassProxy))
+			{
+				OutProxy = PassProxy;
+
+				// Next pass should receive the output of the current child pass.
+				BasePassInput.Set<const FCompositeCorePassProxy*>(PassProxy);
+
+				if (Pass->NeedsSceneTextures())
+				{
+					InContext.bNeedsSceneTextures = true;
+				}
+			}
+		}
+	}
+
+	return OutProxy != nullptr;
+}
+
diff --git a/Engine/Plugins/Experimental/Compositing/Composite/Source/Composite/Public/CompositeActor.h b/Engine/Plugins/Experimental/Compositing/Composite/Source/Composite/Public/CompositeActor.h
--- a/Engine/Plugins/Experimental/Compositing/Composite/Source/Composite/Public/CompositeActor.h
+++ b/Engine/Plugins/Experimental/Compositing/Composite/Source/Composite/Public/CompositeActor.h
@@ -29,6 +29,15 @@ struct FSceneCaptureComponentArray
 	TArray<TObjectPtr<USceneCaptureComponent2D>> Components;
 };
 
+/* Main render mode override. */
+UENUM(BlueprintType)
+enum class ECompositeOverrideMainRenderMode : uint8
+{
+	None UMETA(DisplayName = "None", ToolTip = "Normal rendering through post-processing's tonemap, equivalent to the default \"Final Color (LDR) in RGB\"."),
+	FinalColorHDR UMETA(DisplayName = "Linear", ToolTip = "Equivalent to \"Final Color (HDR) in Linear Working Color Space\"."),
+	FinalToneCurveHDR UMETA(DisplayName = "Linear With Tone Curve", ToolTip = "Equivalent to \"Final Color (with tone curve) in Linear sRGB gamut\"."),
+};
+
 /** Actor used to control properties of the composite pipeline. */
 UCLASS(MinimalAPI)
 class ACompositeActor : public AActor
@@ -116,6 +125,10 @@ public:
 	UPROPERTY(EditAnywhere, BlueprintReadWrite, Interp, Category = "Composite", AdvancedDisplay, meta = (EditCondition = "bOverridesViewUserFlags"))
 	int32 ViewUserFlags;
 
+	/** Override render mode override. */
+	UPROPERTY(EditAnywhere, BlueprintReadWrite, Category = "Composite", AdvancedDisplay)
+	ECompositeOverrideMainRenderMode OverrideMainRenderMode;
+
 public:
 	/** Convenience function used by layers to request a scene capture component. */
 	template<class RetType>
diff --git a/Engine/Plugins/Experimental/Compositing/Composite/Source/Composite/Public/Layers/CompositeLayerMainRender.h b/Engine/Plugins/Experimental/Compositing/Composite/Source/Composite/Public/Layers/CompositeLayerMainRender.h
--- a/Engine/Plugins/Experimental/Compositing/Composite/Source/Composite/Public/Layers/CompositeLayerMainRender.h
+++ b/Engine/Plugins/Experimental/Compositing/Composite/Source/Composite/Public/Layers/CompositeLayerMainRender.h
@@ -8,6 +8,9 @@
 
 #define UE_API COMPOSITE_API
 
+/**
+* Main render layer fetching scene color at current post-processing location (SSR Input and/or After Tonemap).
+*/
 UCLASS(MinimalAPI, BlueprintType, Blueprintable, EditInlineNew, CollapseCategories, meta = (DisplayName = "Main Render"))
 class UCompositeLayerMainRender : public UCompositeLayerBase
 {
diff --git a/Engine/Plugins/Experimental/Compositing/Composite/Source/Composite/Public/Layers/CompositeLayerProcessing.h b/Engine/Plugins/Experimental/Compositing/Composite/Source/Composite/Public/Layers/CompositeLayerProcessing.h
new file mode 100644
--- /dev/null
+++ b/Engine/Plugins/Experimental/Compositing/Composite/Source/Composite/Public/Layers/CompositeLayerProcessing.h
@@ -0,0 +1,39 @@
+// Copyright Epic Games, Inc. All Rights Reserved.
+
+#pragma once
+
+#include "CompositeLayerBase.h"
+
+#include "CompositeLayerProcessing.generated.h"
+
+#define UE_API COMPOSITE_API
+
+/**
+* Special type of layer for processing the output of previous layers. It does NOT merge an additional input.
+*/
+UCLASS(MinimalAPI, BlueprintType, Blueprintable, EditInlineNew, CollapseCategories, meta = (DisplayName = "Processing"))
+class UCompositeLayerProcessing : public UCompositeLayerBase
+{
+	GENERATED_BODY()
+
+public:
+	/** Constructor. */
+	UE_API UCompositeLayerProcessing(const FObjectInitializer& ObjectInitializer);
+	
+	/** Destructor. */
+	UE_API virtual ~UCompositeLayerProcessing();
+
+	/** Getter function to override, returning pass layer proxies to be passed to the render thread. (Proxy objects should be allocated from the provided allocator.) */
+	UE_API virtual bool GetProxy(FTraversalContext& InContext, FSceneRenderingBulkObjectAllocator& InFrameAllocator, FCompositeCorePassProxy*& OutProxy) const override;
+
+#if WITH_EDITOR
+	UE_API virtual bool CanEditChange(const FProperty* InProperty) const;
+	UE_API virtual void PostEditChangeProperty(FPropertyChangedEvent& PropertyChangedEvent) override;
+#endif
+
+	/** Sub-passes applied onto the previous layer. */
+	UPROPERTY(EditAnywhere, BlueprintReadWrite, Instanced, Category = "")
+	TArray<TObjectPtr<UCompositePassBase>> LayerPasses;
+};
+
+#undef UE_API
diff --git a/Engine/Plugins/Experimental/Compositing/Composite/Source/Composite/Public/Layers/CompositeLayerSceneCapture.h b/Engine/Plugins/Experimental/Compositing/Composite/Source/Composite/Public/Layers/CompositeLayerSceneCapture.h
--- a/Engine/Plugins/Experimental/Compositing/Composite/Source/Composite/Public/Layers/CompositeLayerSceneCapture.h
+++ b/Engine/Plugins/Experimental/Compositing/Composite/Source/Composite/Public/Layers/CompositeLayerSceneCapture.h
@@ -10,6 +10,9 @@
 
 class AActor;
 
+/**
+* Layer for scene capture (or custom render pass) render integration.
+*/
 UCLASS(MinimalAPI, BlueprintType, Blueprintable, EditInlineNew, CollapseCategories, meta = (DisplayName = "Scene Capture"))
 class UCompositeLayerSceneCapture : public UCompositeLayerBase
 {
diff --git a/Engine/Plugins/Experimental/Compositing/Composite/Source/Composite/Public/Layers/CompositeLayerShadowReflection.h b/Engine/Plugins/Experimental/Compositing/Composite/Source/Composite/Public/Layers/CompositeLayerShadowReflection.h
--- a/Engine/Plugins/Experimental/Compositing/Composite/Source/Composite/Public/Layers/CompositeLayerShadowReflection.h
+++ b/Engine/Plugins/Experimental/Compositing/Composite/Source/Composite/Public/Layers/CompositeLayerShadowReflection.h
@@ -21,6 +21,10 @@ enum class ECompositeHiddenInSceneCaptureConfiguration : uint8
 	Hidden UMETA(ToolTip = "The following properties are set to true: bHiddenInSceneCapture, bAffectIndirectLightingWhileHidden & bCastHiddenShadow."),
 };
 
+/**
+* Layer for shadow & reflection catching with two additional scene capture renders.
+* Primitives are rendered in one scene capture but not the other. The division of both creates a multiplicative matte that can be applied onto the plate.
+*/
 UCLASS(MinimalAPI, BlueprintType, Blueprintable, EditInlineNew, CollapseCategories, meta = (DisplayName = "Shadow Reflection Catcher"))
 class UCompositeLayerShadowReflection : public UCompositeLayerBase
 {
diff --git a/Engine/Plugins/Experimental/Compositing/CompositeCore/Source/CompositeCore/Private/CompositeCoreSceneViewExtension.cpp b/Engine/Plugins/Experimental/Compositing/CompositeCore/Source/CompositeCore/Private/CompositeCoreSceneViewExtension.cpp
--- a/Engine/Plugins/Experimental/Compositing/CompositeCore/Source/CompositeCore/Private/CompositeCoreSceneViewExtension.cpp
+++ b/Engine/Plugins/Experimental/Compositing/CompositeCore/Source/CompositeCore/Private/CompositeCoreSceneViewExtension.cpp
@@ -163,11 +163,14 @@ void FCompositeCoreSceneViewExtension::SetRenderWork_GameThread(UE::CompositeCor
 				SVE->RenderWork_RenderThread = MoveTemp(FrameWork);
 			}
 		});
+
+	MainRenderMode = InWork.MainRenderMode;
 }
 
 void FCompositeCoreSceneViewExtension::ResetRenderWork_GameThread()
 {
 	bHasCustomRenderWork = false;
+	MainRenderMode.Reset();
 
 	ENQUEUE_RENDER_COMMAND(CopyCompositeCoreRenderWork)(
 		[WeakThis = AsWeak()](FRHICommandListImmediate& RHICmdList)
@@ -237,6 +240,11 @@ void FCompositeCoreSceneViewExtension::SetupViewFamily(FSceneViewFamily& InViewF
 			Iter.RemoveCurrent();
 		}
 	}
+
+	if (MainRenderMode.IsSet())
+	{
+		InViewFamily.SceneCaptureSource = MainRenderMode.GetValue();
+	}
 }
 
 void FCompositeCoreSceneViewExtension::SetupView(FSceneViewFamily& InViewFamily, FSceneView& InView)
diff --git a/Engine/Plugins/Experimental/Compositing/CompositeCore/Source/CompositeCore/Private/CompositeCoreSceneViewExtension.h b/Engine/Plugins/Experimental/Compositing/CompositeCore/Source/CompositeCore/Private/CompositeCoreSceneViewExtension.h
--- a/Engine/Plugins/Experimental/Compositing/CompositeCore/Source/CompositeCore/Private/CompositeCoreSceneViewExtension.h
+++ b/Engine/Plugins/Experimental/Compositing/CompositeCore/Source/CompositeCore/Private/CompositeCoreSceneViewExtension.h
@@ -128,6 +128,9 @@ private:
 	/** Used to keep the SVE active if any custom render work is present.*/
 	bool bHasCustomRenderWork = false;
 
+	/** Optional main render mode override */
+	TOptional<ESceneCaptureSource> MainRenderMode;
+
 	/**
 	* Map of cached proxy output textures per view (post-processing only currently).
 	* 
diff --git a/Engine/Plugins/Experimental/Compositing/CompositeCore/Source/CompositeCore/Private/Passes/CompositeCorePassProxy.cpp b/Engine/Plugins/Experimental/Compositing/CompositeCore/Source/CompositeCore/Private/Passes/CompositeCorePassProxy.cpp
--- a/Engine/Plugins/Experimental/Compositing/CompositeCore/Source/CompositeCore/Private/Passes/CompositeCorePassProxy.cpp
+++ b/Engine/Plugins/Experimental/Compositing/CompositeCore/Source/CompositeCore/Private/Passes/CompositeCorePassProxy.cpp
@@ -61,6 +61,7 @@ namespace UE
 			, PreprocessingPasses{}
 			, FramePasses{}
 			, FrameAllocator{new FSceneRenderingBulkObjectAllocator}
+			, MainRenderMode{}
 		{ }
 
 		bool FRenderWork::IsEmpty() const
@@ -77,6 +78,8 @@ namespace UE
 				bIsEmpty &= Pair.Value.IsEmpty();
 			}
 
+			bIsEmpty &= !MainRenderMode.IsSet();
+
 			return bIsEmpty;
 		}
 
diff --git a/Engine/Plugins/Experimental/Compositing/CompositeCore/Source/CompositeCore/Public/Passes/CompositeCorePassProxy.h b/Engine/Plugins/Experimental/Compositing/CompositeCore/Source/CompositeCore/Public/Passes/CompositeCorePassProxy.h
--- a/Engine/Plugins/Experimental/Compositing/CompositeCore/Source/CompositeCore/Public/Passes/CompositeCorePassProxy.h
+++ b/Engine/Plugins/Experimental/Compositing/CompositeCore/Source/CompositeCore/Public/Passes/CompositeCorePassProxy.h
@@ -226,6 +226,9 @@ namespace UE
 
 			/** Proxy allocator. */
 			TUniquePtr<FSceneRenderingBulkObjectAllocator> FrameAllocator;
+
+			/** Optional main render mode override */
+			TOptional<ESceneCaptureSource> MainRenderMode;
 		};
 	}
 }
