Composite: Quick naming fix for the main render output setting (breaking change).

* Replaced Override + None with "Default" as a simplication.
* Added clarification that "Default" is usually sRGB with tone curve.

diff --git a/Engine/Plugins/Experimental/Compositing/Composite/Source/Composite/Private/CompositeActor.cpp b/Engine/Plugins/Experimental/Compositing/Composite/Source/Composite/Private/CompositeActor.cpp
--- a/Engine/Plugins/Experimental/Compositing/Composite/Source/Composite/Private/CompositeActor.cpp
+++ b/Engine/Plugins/Experimental/Compositing/Composite/Source/Composite/Private/CompositeActor.cpp
@@ -178,12 +178,12 @@ namespace UE
 				RenderWork.ExternalInputs = Context.ExternalTextures;
 				RenderWork.PreprocessingPasses = Context.PreprocessingPasses;
 
-				switch (InCompositeActor.OverrideMainRenderMode)
+				switch (InCompositeActor.MainRenderOutput)
 				{
-				case ECompositeOverrideMainRenderMode::FinalColorHDR:
+				case ECompositeMainRenderOutputMode::FinalColorHDR:
 					RenderWork.MainRenderMode = ESceneCaptureSource::SCS_FinalColorHDR;
 					break;
-				case ECompositeOverrideMainRenderMode::FinalToneCurveHDR:
+				case ECompositeMainRenderOutputMode::FinalToneCurveHDR:
 					RenderWork.MainRenderMode = ESceneCaptureSource::SCS_FinalToneCurveHDR;
 					break;
 				default:
diff --git a/Engine/Plugins/Experimental/Compositing/Composite/Source/Composite/Public/CompositeActor.h b/Engine/Plugins/Experimental/Compositing/Composite/Source/Composite/Public/CompositeActor.h
--- a/Engine/Plugins/Experimental/Compositing/Composite/Source/Composite/Public/CompositeActor.h
+++ b/Engine/Plugins/Experimental/Compositing/Composite/Source/Composite/Public/CompositeActor.h
@@ -29,13 +29,13 @@ struct FSceneCaptureComponentArray
 	TArray<TObjectPtr<USceneCaptureComponent2D>> Components;
 };
 
-/* Main render mode override. */
+/* Main render color space & encoding output mode. */
 UENUM(BlueprintType)
-enum class ECompositeOverrideMainRenderMode : uint8
+enum class ECompositeMainRenderOutputMode : uint8
 {
-	None UMETA(DisplayName = "None", ToolTip = "Normal rendering through post-processing's tonemap, equivalent to the default \"Final Color (LDR) in RGB\"."),
-	FinalColorHDR UMETA(DisplayName = "Linear", ToolTip = "Equivalent to \"Final Color (HDR) in Linear Working Color Space\"."),
-	FinalToneCurveHDR UMETA(DisplayName = "Linear With Tone Curve", ToolTip = "Equivalent to \"Final Color (with tone curve) in Linear sRGB gamut\"."),
+	Default UMETA(DisplayName = "Default", ToolTip = "Default output from post-processing's tonemap, usually sRGB display output with tone curve, equivalent to Final Color (LDR) in RGB."),
+	FinalColorHDR UMETA(DisplayName = "HDR Linear", ToolTip = "Linear output from post-processing's tonemap, equivalent to Final Color (HDR) in Linear Working Color Space."),
+	FinalToneCurveHDR UMETA(DisplayName = "HDR Linear with Tone Curve", ToolTip = "Linear with tone curve output from post-processing's tonemap, equivalent to Final Color (with tone curve) in Linear sRGB gamut."),
 };
 
 /** Actor used to control properties of the composite pipeline. */
@@ -125,9 +125,9 @@ public:
 	UPROPERTY(EditAnywhere, BlueprintReadWrite, Interp, Category = "Composite", AdvancedDisplay, meta = (EditCondition = "bOverridesViewUserFlags"))
 	int32 ViewUserFlags;
 
-	/** Override render mode override. */
+	/** Main render color space & encoding output mode. */
 	UPROPERTY(EditAnywhere, BlueprintReadWrite, Category = "Composite", AdvancedDisplay)
-	ECompositeOverrideMainRenderMode OverrideMainRenderMode;
+	ECompositeMainRenderOutputMode MainRenderOutput;
 
 public:
 	/** Convenience function used by layers to request a scene capture component. */
