Composite: Fix delta keyer parameter tweaking issues in editor.

* Only update the keyer color once commited to prevent in-editor flickering.
* Disable editor throttling for numerical slider issues so that results can be observed in real-time.

diff --git a/Engine/Plugins/Experimental/Compositing/Composite/Source/Composite/Composite.Build.cs b/Engine/Plugins/Experimental/Compositing/Composite/Source/Composite/Composite.Build.cs
--- a/Engine/Plugins/Experimental/Compositing/Composite/Source/Composite/Composite.Build.cs
+++ b/Engine/Plugins/Experimental/Compositing/Composite/Source/Composite/Composite.Build.cs
@@ -41,6 +41,7 @@ public class Composite : ModuleRules
 				{
 					"UnrealEd",
 					"Slate",
+					"SlateCore",
 				}
 			);
 		}
diff --git a/Engine/Plugins/Experimental/Compositing/Composite/Source/Composite/Private/Passes/CompositePassDeltaKeyer.cpp b/Engine/Plugins/Experimental/Compositing/Composite/Source/Composite/Private/Passes/CompositePassDeltaKeyer.cpp
--- a/Engine/Plugins/Experimental/Compositing/Composite/Source/Composite/Private/Passes/CompositePassDeltaKeyer.cpp
+++ b/Engine/Plugins/Experimental/Compositing/Composite/Source/Composite/Private/Passes/CompositePassDeltaKeyer.cpp
@@ -7,6 +7,10 @@
 #include "PixelShaderUtils.h"
 #include "TextureResource.h"
 
+#if WITH_EDITOR
+#include "Application/ThrottleManager.h"
+#endif
+
 DECLARE_GPU_STAT_NAMED(FCompositeDeltaKeyer, TEXT("Composite.DeltaKeyer"));
 
 class FCompositePassDeltaKeyerShader : public FGlobalShader
@@ -71,8 +75,6 @@ namespace UE
 					Output = CreateOutputRenderTarget(GraphBuilder, InView, PassContext.OutputViewRect, Input.Texture->Desc, TEXT("DeltaKeyerCompositePass"));
 				}
 
-				//CleanPlateWeakPtr.get()
-
 				FGlobalShaderMap* GlobalShaderMap = GetGlobalShaderMap(InView.GetFeatureLevel());
 				const FIntPoint TextureSize = Input.Texture->Desc.Extent;
 
@@ -161,3 +163,49 @@ bool UCompositePassDeltaKeyer::GetProxy(const UE::CompositeCore::FPassInputDecl&
 	return true;
 }
 
+#if WITH_EDITOR
+void UCompositePassDeltaKeyer::PreEditChange(FProperty* PropertyThatWillChange)
+{
+	Super::PreEditChange(PropertyThatWillChange);
+
+	if (!PropertyThatWillChange)
+	{
+		return;
+	}
+
+	const FName PropertyName = PropertyThatWillChange->GetFName();
+
+	if (PropertyName == GET_MEMBER_NAME_CHECKED(ThisClass, ColorTolerance)
+		|| PropertyName == GET_MEMBER_NAME_CHECKED(ThisClass, EdgeSoftness)
+		|| PropertyName == GET_MEMBER_NAME_CHECKED(ThisClass, AlphaClipRange)
+		|| PropertyName == GET_MEMBER_NAME_CHECKED(ThisClass, DespillStrength)
+		|| PropertyName == GET_MEMBER_NAME_CHECKED(ThisClass, bInvertAlpha)
+		)
+	{
+		/**
+		* For the above properties, we want to keep the editor ticking so that
+		* result from sliding their (SNumericEntryBox) values can be observed in real-time.
+		* 
+		* Unlike component transform customizations which prevent throttling explicitly
+		* (see .PreventThrottling(true)) we don't appear to have a metadata property for this.
+		* 
+		* As an alternative solution we directly disable throttling from the slate manager,
+		* and re-enable it once the value is set.
+		*/
+		FSlateThrottleManager::Get().DisableThrottle(true);
+		bEditorThrottleDisabled = true;
+	}
+}
+
+void UCompositePassDeltaKeyer::PostEditChangeProperty(FPropertyChangedEvent& PropertyChangedEvent)
+{
+	Super::PostEditChangeProperty(PropertyChangedEvent);
+
+	if (bEditorThrottleDisabled && PropertyChangedEvent.ChangeType != EPropertyChangeType::Interactive)
+	{
+		FSlateThrottleManager::Get().DisableThrottle(false);
+		bEditorThrottleDisabled = false;
+	}
+}
+#endif
+
diff --git a/Engine/Plugins/Experimental/Compositing/Composite/Source/Composite/Public/Passes/CompositePassDeltaKeyer.h b/Engine/Plugins/Experimental/Compositing/Composite/Source/Composite/Public/Passes/CompositePassDeltaKeyer.h
--- a/Engine/Plugins/Experimental/Compositing/Composite/Source/Composite/Public/Passes/CompositePassDeltaKeyer.h
+++ b/Engine/Plugins/Experimental/Compositing/Composite/Source/Composite/Public/Passes/CompositePassDeltaKeyer.h
@@ -35,9 +35,14 @@ public:
 
 	UE_API virtual bool GetProxy(const UE::CompositeCore::FPassInputDecl& InputDecl, FSceneRenderingBulkObjectAllocator& InFrameAllocator, FCompositeCorePassProxy*& OutProxy) const override;
 
+#if WITH_EDITOR
+	UE_API virtual void PreEditChange(FProperty* PropertyThatWillChange) override;
+	UE_API virtual void PostEditChangeProperty(FPropertyChangedEvent& PropertyChangedEvent) override;
+#endif
+
 public:
 	/** Background key color. */
-	UPROPERTY(EditAnywhere, BlueprintReadWrite, Category = "", meta = (EditCondition = "CleanPlate == nullptr"))
+	UPROPERTY(EditAnywhere, BlueprintReadWrite, Category = "", meta = (OnlyUpdateOnInteractionEnd, EditCondition = "CleanPlate == nullptr"))
 	FLinearColor KeyColor;
 
 	/**
@@ -70,6 +75,11 @@ public:
 	/** Invert the alpha key. */
 	UPROPERTY(EditAnywhere, BlueprintReadWrite, Category = "")
 	bool bInvertAlpha;
+
+#if WITH_EDITOR
+	/** Flag tracking whether we disabled throttling to later restore it. */
+	bool bEditorThrottleDisabled = false;
+#endif
 };
 
 #undef UE_API
