Composite: Improvements for delta keyer pass (with breaking changes).

* Removed use of smoothstep in favor of linear interpolation
* Improved default despill
* Added fringe gamma control
* Added alpha visualization toggle
* Added additional functions for testing

diff --git a/Engine/Plugins/Experimental/Compositing/Composite/Shaders/Private/CompositeDeltaKeyer.usf b/Engine/Plugins/Experimental/Compositing/Composite/Shaders/Private/CompositeDeltaKeyer.usf
--- a/Engine/Plugins/Experimental/Compositing/Composite/Shaders/Private/CompositeDeltaKeyer.usf
+++ b/Engine/Plugins/Experimental/Compositing/Composite/Shaders/Private/CompositeDeltaKeyer.usf
@@ -1,21 +1,22 @@
 // Copyright Epic Games, Inc. All Rights Reserved.
 
 #include "/Engine/Private/Common.ush"
-
-#include "/Engine/Private/ACES/ACESCommon.ush"
-#include "/Engine/Private/ColorSpace.ush"
-#include "/Engine/Private/GammaCorrectionCommon.ush"
 #include "/Engine/Private/ScreenPass.ush"
 #include "/Engine/Private/PostProcessCommon.ush"
 
+#include "CompositeDeltaKeyerCommon.ush"
+
 #ifndef TILE_SIZE
     #error TILE_SIZE must be defined
 #endif
 
+// Hand-picked scale factor used to adjust LUV color delta range
+static const float kDistanceScaleFactor_CIELUV = 2.0 * DeltaE_JND / Lstar_CIELUV;
+
 #define DESPILL_METHOD_NONE 0
-#define DESPILL_METHOD_GREEN_AVG 1
-#define DESPILL_METHOD_BLUE_AVG 2
-#define DESPILL_METHOD_ALPHA_KEY 3
+#define DESPILL_METHOD_DEFAULT 1
+#define DESPILL_METHOD_GREEN_AVG 2
+#define DESPILL_METHOD_BLUE_AVG 3
 
 Texture2D InputTexture;
 SamplerState InputSampler;
@@ -31,108 +32,29 @@ SCREEN_PASS_TEXTURE_VIEWPORT(Output)
 float3 KeyColor;
 float4 Params0;
 float4 Params1;
-uint bInvertAlpha;
 uint DespillMethod;
+uint bVisualizeAlpha;
+uint bInvertAlpha;
 
-// CIE 1931 D65 white point
-static const float3 XYZ_n = xyY_2_XYZ(float3(0.3127, 0.3290, 1.0));
-
-static const float DeltaE_JND = 1.0;
-static const float Lstar_CIELUV = 100.0;
-// Hand-picked scale factor used to adjust LUV color delta range
-static const float kDistanceScaleFactor_CIELUV = 5.0 * DeltaE_JND / Lstar_CIELUV;
-
-float SafeRcp(float x)
-{
-    return x > 0.0 ? rcp(x) : 0.0;
-}
-
-// Convert color from XYZ to CIE 1976 L*, u*, v* color space.
-// We use the latter for its improved perceptual uniformity, as a way to measure chromaticity distance for keying.
-float3 CIE_XYZ_to_LUV(float3 XYZ)
+// Output color function with debug convenience
+float4 Output(float4 Color)
 {
-	const float Epsilon = 216.0f / 24389.0f; // (6/29)^3
-	const float Kappa = 24389.0f / 27.0f; // (29/3)^3
+	// Optionally invert alpha
+	Color.a = select(bInvertAlpha > 0, 1.0f - Color.a, Color.a);
 
-	float L = select(XYZ[1] / XYZ_n[1] > Epsilon, 116 * pow(XYZ[1]/ XYZ_n[1], 1.0/3.0) - 16, Kappa * XYZ[1] / XYZ_n[1]);
-
-	if (L == 0.0f)
+	if(bVisualizeAlpha > 0)
 	{
-		return float3(0.0f, 0.0f, 0.0f);
+		// Optionally show the alpha key matte
+		Color.rgb = sRGBToLinear(float3(Color.a, Color.a, Color.a));
+		Color.a = 1.0f;
+	}
+	else
+	{
+		// Return premultiplied alpha
+		Color.rgb *= Color.a;
 	}
-	
-	float Denom = XYZ[0] + 15 * XYZ[1] + 3 * XYZ[2];
-	float Denom_n = XYZ_n[0] + 15 * XYZ_n[1] + 3 * XYZ_n[2];
-	
-	float up = select(Denom > 0.0f, 4 * XYZ[0] / Denom, 0.0f);
-	float vp = select(Denom > 0.0f, 9 * XYZ[1] / Denom, 0.0f);
-
-	float u_n = 4 * XYZ_n[0] / Denom_n;
-	float v_n = 9 * XYZ_n[1] / Denom_n;
-
-	float u = 13 * L * (up - u_n);
-	float v = 13 * L * (vp - v_n);
-
-	return float3(L, u, v);
-}
-
-float CalcColorDelta(float3 ColorA, float3 ColorB)
-{
-	float3 LUV_A = CIE_XYZ_to_LUV(LinearRGB_2_XYZ(ColorA));
-	float3 LUV_B = CIE_XYZ_to_LUV(LinearRGB_2_XYZ(ColorB));
-	
-	// We only calculate the distance in the (luminance-relative) chromaticity dimensions
-	float Delta = distance(LUV_A.yz, LUV_B.yz);
-
-	return kDistanceScaleFactor_CIELUV * Delta;
-}
-
-float CalculateAlpha(float Edge0, float Edge1, float Distance)
-{
-	// Colors within the specified tolerance are fully translucent, with a soft transition depending on the edge smoothness
-	float Alpha = smoothstep(Edge0, Edge1, Distance);
-
-	const float AlphaMin = Params0.z;
-	const float AlphaMax = Params0.w;
-	
-	// Apply alpha clip range (equivalent to saturate by default)
-	Alpha = select(Alpha < AlphaMin, 0.0f, Alpha);
-	Alpha = select(Alpha > AlphaMax, 1.0f, Alpha);
-
-	// Apply optional alpha inversion
-	Alpha = select(bInvertAlpha > 0, 1.0f - Alpha, Alpha);
-
-	return Alpha;
-}
-
-float3 AlphaKeyDespill(float3 Color, float AlphaMask)
-{
-    float KeyColorRcp = SafeRcp(dot(KeyColor, KeyColor));
-    float Weight = AlphaMask * dot(Color.rgb, KeyColor) * KeyColorRcp;
-
-	return max(0.0f, Color.rgb - Weight * KeyColor.rgb);
-}
-
-float3 GreenDespillAvg(float3 Color)
-{
-	float Avg = 0.5f * (Color.r + Color.b);
-
-	return float3(
-		Color.r,
-		select(Color.g > Avg, Avg, Color.g),
-		Color.b
-	);
-}
-
-float3 BlueDespillAvg(float3 Color)
-{
-	float Avg = 0.5f * (Color.r + Color.g);
 
-	return float3(
-		Color.r,
-		Color.g,
-		select(Color.b > Avg, Avg, Color.b)
-	);
+	return Color;
 }
 
 [numthreads(TILE_SIZE, TILE_SIZE, 1)]
@@ -149,55 +71,62 @@ void MainCS(
 	{
 		return;
 	}
-	
-	float4 Color = InputTexture.SampleLevel(InputSampler, UV, 0);
-
-	// Apply keying
-	const float ColorTolerance = Params0.x;
-	const float EdgeSmoothness = Params0.y;
 
+	// Input parameters
+	const float DeltaThresholdMin = Params0.x;
+	const float DeltaThresholdMax = Params0.y;
+	const float AlphaMin = Params0.z;
+	const float AlphaMax = Params0.w;
+	const float DespillStrength = Params1.x;
+	const float FringeGamma = Params1.y;
+	const float ChromaticityOnly = Params1.z;
+	
+	float4 SourceColor = InputTexture.SampleLevel(InputSampler, UV, 0);
 #if USE_CLEAN_PLATE
-	float3 PlateColor = CleanPlateTexture.SampleLevel(CleanPlateSampler, UV, 0).rgb;
+	float3 Background = CleanPlateTexture.SampleLevel(CleanPlateSampler, UV, 0).rgb;
 #else
-	float3 PlateColor = KeyColor;
+	float3 Background = KeyColor;
 #endif
 
-	// Euclidean distance in (L)uv space as a measure of the color difference
-	const float Delta = CalcColorDelta(Color.rgb, PlateColor);
-	// Apply main delta key
-	Color.a = CalculateAlpha(ColorTolerance, ColorTolerance + EdgeSmoothness, Delta);
+	// Map the source color onto what we call the "foreground".
+	// Convert both to Luv color space
+	float3 Foreground_Luv = RGB_to_Luv(SourceColor.rgb);
+	float3 Background_Luv = RGB_to_Luv(Background);
+
+	// Optionally preserve or discard the lightness contribution by shrinking its component distance
+	Background_Luv[0] = lerp(Background_Luv[0], Foreground_Luv[0], ChromaticityOnly);
+	
+    // Calculate chromaticity difference & alpha
+    float Delta = kDistanceScaleFactor_CIELUV * distance(Foreground_Luv, Background_Luv);
+    float Alpha = Linearstep(DeltaThresholdMin, DeltaThresholdMax, Delta);
+	
+	// Project the foreground color onto the background in uv-space
+	float2 ProjectedDiff = ProjectClamped(Foreground_Luv.yz, Background_Luv.yz);
 
-	// Apply despill
-	if(DespillMethod > DESPILL_METHOD_NONE)
+	if(DespillMethod == DESPILL_METHOD_DEFAULT)
 	{
-		const float DespillStrength = Params1.x;
-		float3 Despilled;
-
-		if(DespillMethod == DESPILL_METHOD_GREEN_AVG)
-		{
-			Despilled = GreenDespillAvg(Color.rgb);
-		}
-		else if(DespillMethod == DESPILL_METHOD_BLUE_AVG)
-		{
-			Despilled = BlueDespillAvg(Color.rgb);
-		}
-		else if(DespillMethod == DESPILL_METHOD_ALPHA_KEY)
-		{
-			// Reduced delta such that the edge regions receive more despill (TODO: parametrize)
-			const float AdjustedDelta = 0.75 * Delta;
-			const float AlphaMask = 1.0 - CalculateAlpha(ColorTolerance, ColorTolerance + EdgeSmoothness, AdjustedDelta);
-			Despilled = AlphaKeyDespill(Color.rgb, AlphaMask);
-		}
-
-		// Re-insert luminance to prevent any darkening caused by despill
-		Despilled += Luminance(max(0.0f, Color.rgb - Despilled));
-
-		Color.rgb = lerp(Color.rgb, Despilled, DespillStrength);
+		// Despill the foreground
+		Foreground_Luv.yz -= DespillStrength * ProjectedDiff;
 	}
 
-	// Return premultiplied alpha for AlphaComposite material blend mode
-	Color.rgb *= Color.a;
+	// Apply alpha threshold range
+	float AdjustedAlpha = Linearstep(AlphaMin, AlphaMax, Alpha);
+
+    // Convert back to RGB
+	float4 FinalColor = float4(Luv_to_RGB(Foreground_Luv), AdjustedAlpha);
+
+	if(DespillMethod == DESPILL_METHOD_GREEN_AVG)
+	{
+		FinalColor.rgb = GreenDespillAvg(FinalColor.rgb);
+	}
+	else if(DespillMethod == DESPILL_METHOD_BLUE_AVG)
+	{
+		FinalColor.rgb = BlueDespillAvg(FinalColor.rgb);
+	}
+	
+	// Darken fringe
+	FinalColor.rgb = FinalColor.rgb * pow(Alpha, FringeGamma);
 	
-	uint2 PixelPos = DispatchThreadId + Output_ViewportMin;
-	RWOutputTexture[PixelPos] = Color;
+	const uint2 PixelPos = DispatchThreadId + Output_ViewportMin;
+	RWOutputTexture[PixelPos] = Output(FinalColor);
 }
diff --git a/Engine/Plugins/Experimental/Compositing/Composite/Shaders/Private/CompositeDeltaKeyerCommon.ush b/Engine/Plugins/Experimental/Compositing/Composite/Shaders/Private/CompositeDeltaKeyerCommon.ush
new file mode 100644
--- /dev/null
+++ b/Engine/Plugins/Experimental/Compositing/Composite/Shaders/Private/CompositeDeltaKeyerCommon.ush
@@ -0,0 +1,194 @@
+#pragma once
+
+#include "/Engine/Private/ACES/ACESCommon.ush"
+#include "/Engine/Private/ColorSpace.ush"
+#include "/Engine/Private/GammaCorrectionCommon.ush"
+
+// CIE 1931 D65 white point
+static const float3 XYZ_n = xyY_2_XYZ(float3(0.3127, 0.3290, 1.0));
+
+static const float ZeroEpsilon = 1.0e-6;
+static const float DeltaE_JND = 1.0;
+static const float Lstar_CIELUV = 100.0;
+static const float Lstar_CIELUV_Inv = 1.0 / Lstar_CIELUV;
+
+// Safe reciprocal
+float SafeRcp(float x)
+{
+	return x > 0.0 ? rcp(x) : 0.0;
+}
+
+float LuminanceSRGB(float3 InColor)
+{
+	return dot(InColor, float3(0.2126390059, 0.7151686788, 0.0721923154));
+}
+
+// Linear equivalent of smoothstep
+float Linearstep(float MinValue, float MaxValue, float X)
+{
+	return saturate((X - MinValue) / (MaxValue - MinValue));
+}
+
+float2 Project(float2 VecA, float2 VecB)
+{
+	float Dot = max(0.0f, dot(VecA, VecB));
+	float LengthSquaredB = dot(VecB, VecB);
+	
+	return VecB * (Dot / LengthSquaredB);
+}
+
+// Project vector A onto B, and return 0 if the projection was negative
+float2 ProjectClamped(float2 VecA, float2 VecB)
+{
+	float Dot = max(0.0f, dot(VecA, VecB));
+	float LengthSquaredB = dot(VecB, VecB);
+	
+	return select(LengthSquaredB < ZeroEpsilon, float2(0, 0), VecB * (Dot / LengthSquaredB));
+}
+
+float3 ProjectClamped(float3 VecA, float3 VecB)
+{
+	float Dot = max(0.0f, dot(VecA, VecB));
+	float LengthSquaredB = dot(VecB, VecB);
+
+	return select(LengthSquaredB < ZeroEpsilon, float3(0, 0, 0), VecB * (Dot / LengthSquaredB));
+}
+
+// Convert color from XYZ to CIE 1976 L*, u*, v* color space.
+// We use the latter for its improved perceptual uniformity, as a way to measure chromaticity distance for keying.
+float3 CIE_XYZ_to_Luv(float3 XYZ)
+{
+	const float Epsilon = 216.0f / 24389.0f; // (6/29)^3
+	const float Kappa = 24389.0f / 27.0f; // (29/3)^3
+
+	float L = select(XYZ[1] / XYZ_n[1] > Epsilon, 116 * pow(XYZ[1]/ XYZ_n[1], 1.0/3.0) - 16, Kappa * XYZ[1] / XYZ_n[1]);
+
+	if (L < ZeroEpsilon)
+	{
+		return float3(0.0f, 0.0f, 0.0f);
+	}
+	
+	float Denom = XYZ[0] + 15 * XYZ[1] + 3 * XYZ[2];
+	float Denom_n = XYZ_n[0] + 15 * XYZ_n[1] + 3 * XYZ_n[2];
+	
+	float up = select(Denom > 0.0f, 4 * XYZ[0] / Denom, 0.0f);
+	float vp = select(Denom > 0.0f, 9 * XYZ[1] / Denom, 0.0f);
+
+	float u_n = 4 * XYZ_n[0] / Denom_n;
+	float v_n = 9 * XYZ_n[1] / Denom_n;
+
+	float u = 13 * L * (up - u_n);
+	float v = 13 * L * (vp - v_n);
+
+	return float3(L, u, v);
+}
+
+// Convert color from CIE 1976 L*, u*, v* to XYZ color space.
+float3 CIE_Luv_to_XYZ(float3 Luv)
+{
+	const float Epsilon = 216.0f / 24389.0f; // (6/29)^3
+	const float Kappa = 24389.0f / 27.0f; // (29/3)^3
+
+	float L = Luv[0];
+	float u = Luv[1];
+	float v = Luv[2];
+
+	// Handle the case where L is zero, which corresponds to black.
+	if (L < ZeroEpsilon)
+	{
+		return float3(0.0f, 0.0f, 0.0f);
+	}	
+
+	float Y_ratio = select(L > Kappa * Epsilon, pow((L + 16.0f) / 116.0f, 3.0f), L / Kappa);
+	float Y = Y_ratio * XYZ_n[1];
+
+	float Denom_n = XYZ_n[0] + 15.0f * XYZ_n[1] + 3.0f * XYZ_n[2];
+	float u_n = 4.0f * XYZ_n[0] / Denom_n;
+	float v_n = 9.0f * XYZ_n[1] / Denom_n;
+	
+	float up = u / (13.0f * L) + u_n;
+	float vp = v / (13.0f * L) + v_n;
+
+	float X = Y * (9.0f * up) / (4.0f * vp);
+	float Z = Y * (12.0f - 3.0f * up - 20.0f * vp) / (4.0f * vp);
+	
+	return float3(X, Y, Z);
+}
+
+float3 RGB_to_Luv(float3 InRGB)
+{
+	return CIE_XYZ_to_Luv(LinearRGB_2_XYZ(InRGB));
+}
+
+float3 Luv_to_RGB(float3 InLuv)
+{
+	return XYZ_2_LinearRGB(CIE_Luv_to_XYZ(InLuv));
+}
+
+float3 GreenDespillAvg(float3 Color)
+{
+	float Avg = 0.5f * (Color.r + Color.b);
+
+	return float3(
+		Color.r,
+		select(Color.g > Avg, Avg, Color.g),
+		Color.b
+	);
+}
+
+float3 BlueDespillAvg(float3 Color)
+{
+	float Avg = 0.5f * (Color.r + Color.g);
+
+	return float3(
+		Color.r,
+		Color.g,
+		select(Color.b > Avg, Avg, Color.b)
+	);
+}
+
+// --- Unused but kept for testing ---
+float4 BlueScreenKeyer(
+	float3 Plate,
+	float3 ScreenColor,
+	float  RedInfluence,
+	float  GreenInfluence,
+	float  AlphaGain = 1.0f)
+{
+	// Blue suppression: remove blue proportionally to the size of other channels
+	float PlateB = max(Plate.b - RedInfluence * Plate.r
+		- GreenInfluence * Plate.g, 0.0);
+	float ScreenB = max(ScreenColor.b - RedInfluence * ScreenColor.r
+		- GreenInfluence * ScreenColor.g, 0.0);
+
+	// Ratio of how much blue is left on the plate over the screen
+	float BlueRatio = saturate(PlateB * SafeRcp(ScreenB));
+	float3 ScreenAdj = float3(ScreenColor.r, ScreenColor.g, ScreenB);
+	float3 BlueSpill = clamp(ScreenAdj * BlueRatio, 0.0, 32.0);
+
+	// Final color
+	float3 ResultRGB = clamp(Plate - BlueSpill, 0.0, 1024.0);
+
+	// Alpha matte
+	float LumaDiff = Luminance(Plate - ScreenColor);
+	float Alpha = saturate(max(LumaDiff, 1.0 - BlueRatio) * AlphaGain);
+
+	return float4(ResultRGB, Alpha);
+}
+
+float3 SwizzleGreenBlue(float3 InColor)
+{
+	return float3(InColor.r, InColor.b, InColor.g);
+}
+
+float4 GreenScreenKeyer(float3 Plate, float3 ScreenColor, float RedInfluence, float BlueInfluence, float AlphaGain = 1.0f)
+{
+	float3 SwizzledPlate = SwizzleGreenBlue(Plate);
+	float3 SizzledScreenColor = SwizzleGreenBlue(ScreenColor);
+	
+	float4 Output = BlueScreenKeyer(SwizzledPlate, SizzledScreenColor, RedInfluence, BlueInfluence, AlphaGain);
+	Output.rgb = SwizzleGreenBlue(Output.rgb);
+
+	return Output;
+}
+// --- End Unused ---
\ No newline at end of file
diff --git a/Engine/Plugins/Experimental/Compositing/Composite/Source/Composite/Private/Passes/CompositePassDeltaKeyer.cpp b/Engine/Plugins/Experimental/Compositing/Composite/Source/Composite/Private/Passes/CompositePassDeltaKeyer.cpp
--- a/Engine/Plugins/Experimental/Compositing/Composite/Source/Composite/Private/Passes/CompositePassDeltaKeyer.cpp
+++ b/Engine/Plugins/Experimental/Compositing/Composite/Source/Composite/Private/Passes/CompositePassDeltaKeyer.cpp
@@ -33,8 +33,9 @@ class FCompositePassDeltaKeyerShader : public FGlobalShader
 		SHADER_PARAMETER(FVector3f, KeyColor)
 		SHADER_PARAMETER(FVector4f, Params0)
 		SHADER_PARAMETER(FVector4f, Params1)
-		SHADER_PARAMETER(uint32, DespillMethod)
+		SHADER_PARAMETER(uint32, bVisualizeAlpha)
 		SHADER_PARAMETER(uint32, bInvertAlpha)
+		SHADER_PARAMETER(uint32, DespillMethod)
 	END_SHADER_PARAMETER_STRUCT()
 
 	static inline void ModifyCompilationEnvironment(const FGlobalShaderPermutationParameters& Parameters, FShaderCompilerEnvironment& OutEnvironment)
@@ -92,6 +93,7 @@ namespace UE
 				PassParameters->Params0 = Params0;
 				PassParameters->Params1 = Params1;
 				PassParameters->DespillMethod = DespillMethod;
+				PassParameters->bVisualizeAlpha = static_cast<uint32>(bVisualizeAlpha);
 				PassParameters->bInvertAlpha = static_cast<uint32>(bInvertAlpha);
 	
 				bool bUseCleanPlate = false;
@@ -127,6 +129,8 @@ namespace UE
 			FVector4f Params0 = FVector4f::Zero();
 			FVector4f Params1 = FVector4f::Zero();
 			uint32 DespillMethod = 0;
+			FLinearColor DespillColor = FLinearColor::White;
+			bool bVisualizeAlpha = false;
 			bool bInvertAlpha = false;
 			TWeakObjectPtr<UTexture> CleanPlateWeakPtr;
 		};
@@ -136,11 +140,13 @@ namespace UE
 UCompositePassDeltaKeyer::UCompositePassDeltaKeyer(const FObjectInitializer& ObjectInitializer)
 	: Super(ObjectInitializer)
 	, KeyColor(FLinearColor::Green)
-	, ColorTolerance(1.0f)
-	, EdgeSoftness(1.0f)
-	, AlphaClipRange(0.0f, 1.0f)
-	, DespillMethod(ECompositeDeltaKeyerDespillMethod::None)
+	, ColorTolerance(0.08f, 1.0f)
+	, AlphaThreshold(0.0f, 1.0f)
+	, ChromaticityOnly(1.0f)
+	, DespillMethod(ECompositeDeltaKeyerDespillMethod::Default)
 	, DespillStrength(1.0f)
+	, FringeGamma(2.2f)
+	, bVisualizeAlpha(false)
 	, bInvertAlpha(false)
 {
 }
@@ -153,10 +159,11 @@ bool UCompositePassDeltaKeyer::GetProxy(const UE::CompositeCore::FPassInputDecl&
 
 	FCompositePassDeltaKeyerProxy* Proxy = InFrameAllocator.Create<FCompositePassDeltaKeyerProxy>(UE::CompositeCore::FPassInputDeclArray{ InputDecl });
 	Proxy->KeyColor = KeyColor;
-	Proxy->Params0 = FVector4f(ColorTolerance, EdgeSoftness, AlphaClipRange.X, AlphaClipRange.Y);
-	Proxy->Params1 = FVector4f(DespillStrength, 0.0f, 0.0f, 0.0f);
+	Proxy->Params0 = FVector4f(ColorTolerance.X, ColorTolerance.Y, AlphaThreshold.X, AlphaThreshold.Y);
+	Proxy->Params1 = FVector4f(DespillStrength, FringeGamma, FMath::Clamp(ChromaticityOnly, 0.0f, 1.0f), 0.0f);
 	Proxy->DespillMethod = static_cast<uint32>(DespillMethod);
 	Proxy->bInvertAlpha = bInvertAlpha;
+	Proxy->bVisualizeAlpha = bVisualizeAlpha;
 	Proxy->CleanPlateWeakPtr = CleanPlate;
 
 	OutProxy = Proxy;
@@ -174,16 +181,16 @@ void UCompositePassDeltaKeyer::PreEditChange(FProperty* PropertyThatWillChange)
 	}
 
 	const FName PropertyName = PropertyThatWillChange->GetFName();
+	const bool bIsBooleanProperty = CastField<const FBoolProperty>(PropertyThatWillChange) != nullptr;
 
-	if (PropertyName == GET_MEMBER_NAME_CHECKED(ThisClass, ColorTolerance)
-		|| PropertyName == GET_MEMBER_NAME_CHECKED(ThisClass, EdgeSoftness)
-		|| PropertyName == GET_MEMBER_NAME_CHECKED(ThisClass, AlphaClipRange)
-		|| PropertyName == GET_MEMBER_NAME_CHECKED(ThisClass, DespillStrength)
-		|| PropertyName == GET_MEMBER_NAME_CHECKED(ThisClass, bInvertAlpha)
+	if (PropertyName != GET_MEMBER_NAME_CHECKED(ThisClass, bIsEnabled)
+		&& PropertyName != GET_MEMBER_NAME_CHECKED(ThisClass, KeyColor)
+		&& PropertyName != GET_MEMBER_NAME_CHECKED(ThisClass, CleanPlate)
+		&& !bIsBooleanProperty
 		)
 	{
 		/**
-		* For the above properties, we want to keep the editor ticking so that
+		* For properties not included in the list above, we want to keep the editor ticking so that
 		* result from sliding their (SNumericEntryBox) values can be observed in real-time.
 		* 
 		* Unlike component transform customizations which prevent throttling explicitly
diff --git a/Engine/Plugins/Experimental/Compositing/Composite/Source/Composite/Public/Passes/CompositePassDeltaKeyer.h b/Engine/Plugins/Experimental/Compositing/Composite/Source/Composite/Public/Passes/CompositePassDeltaKeyer.h
--- a/Engine/Plugins/Experimental/Compositing/Composite/Source/Composite/Public/Passes/CompositePassDeltaKeyer.h
+++ b/Engine/Plugins/Experimental/Compositing/Composite/Source/Composite/Public/Passes/CompositePassDeltaKeyer.h
@@ -12,9 +12,9 @@ UENUM(BlueprintType)
 enum class ECompositeDeltaKeyerDespillMethod : uint8
 {
 	None = 0,
-	GreenAverage = 1,
-	BlueAverage = 2,
-	Edge = 3,
+	Default = 1,
+	GreenAverage = 2,
+	BlueAverage = 3,
 };
 
 /**
@@ -52,17 +52,20 @@ public:
 	UPROPERTY(EditAnywhere, BlueprintReadWrite, Category = "")
 	TObjectPtr<UTexture> CleanPlate;
 
-	/**  A threshold for how similar a pixel’s color must be to the Key Color before it is keyed out. */
-	UPROPERTY(EditAnywhere, BlueprintReadWrite, Category = "", meta = (UIMin = 0.0f, UIMax = 2.0f))
-	float ColorTolerance;
-	
-	/** Controls how gradually the transition between keyed and non-keyed areas occurs at the edges. */
+	/**
+	 * A color tolerance range for keying, measured as the distance between the input & key color.
+	 * The minimum tolerance will be fully keyed, while the maximum fully opaque.
+	 */
 	UPROPERTY(EditAnywhere, BlueprintReadWrite, Category = "", meta = (UIMin = 0.0f, UIMax = 2.0f))
-	float EdgeSoftness;
+	FVector2f ColorTolerance;
 
-	/** Clips any alpha value outside the specified range to zero or one respectively. */
+	/** Thresholds any alpha value outside the specified range to zero or one respectively, with linear interpolation in-between. */
 	UPROPERTY(EditAnywhere, BlueprintReadWrite, Category = "", meta = (UIMin = 0.0f, UIMax = 1.0f))
-	FVector2f AlphaClipRange;
+	FVector2f AlphaThreshold;
+
+	/** Whether or not to include lightness/luminance in the color tolerance calculation. */
+	UPROPERTY(EditAnywhere, BlueprintReadWrite, Category = "", meta = (UIMin = 0.0f, UIMax = 1.0f))
+	float ChromaticityOnly;
 
 	/** Despill method, to remove/neutralize the residual key color spill on the subject. */
 	UPROPERTY(EditAnywhere, BlueprintReadWrite, Category = "")
@@ -72,6 +75,14 @@ public:
 	UPROPERTY(EditAnywhere, BlueprintReadWrite, Category = "", meta = (UIMin = 0.0f, UIMax = 1.0f))
 	float DespillStrength;
 
+	/** Gamma adjustment for the edge fringe between full translucency & full opacity. */
+	UPROPERTY(EditAnywhere, BlueprintReadWrite, Category = "", meta = (ClampMin = 0.1f, ClampMax = 4.0))
+	float FringeGamma;
+
+	/** Vizualize the alpha key. */
+	UPROPERTY(EditAnywhere, Transient, Category = "")
+	bool bVisualizeAlpha;
+
 	/** Invert the alpha key. */
 	UPROPERTY(EditAnywhere, BlueprintReadWrite, Category = "")
 	bool bInvertAlpha;
