Composite: Fix missing copyright header.

diff --git a/Engine/Plugins/Experimental/Compositing/Composite/Shaders/Private/CompositeDeltaKeyerCommon.ush b/Engine/Plugins/Experimental/Compositing/Composite/Shaders/Private/CompositeDeltaKeyerCommon.ush
--- a/Engine/Plugins/Experimental/Compositing/Composite/Shaders/Private/CompositeDeltaKeyerCommon.ush
+++ b/Engine/Plugins/Experimental/Compositing/Composite/Shaders/Private/CompositeDeltaKeyerCommon.ush
@@ -1,3 +1,5 @@
+// Copyright Epic Games, Inc. All Rights Reserved.
+
 #pragma once
 
 #include "/Engine/Private/ACES/ACESCommon.ush"
