Composite: Fix double premult in last CL, remove fringe gamma.

diff --git a/Engine/Plugins/Experimental/Compositing/Composite/Shaders/Private/CompositeDeltaKeyer.usf b/Engine/Plugins/Experimental/Compositing/Composite/Shaders/Private/CompositeDeltaKeyer.usf
--- a/Engine/Plugins/Experimental/Compositing/Composite/Shaders/Private/CompositeDeltaKeyer.usf
+++ b/Engine/Plugins/Experimental/Compositing/Composite/Shaders/Private/CompositeDeltaKeyer.usf
@@ -45,7 +45,7 @@ float4 Output(float4 Color)
 	if(bVisualizeAlpha > 0)
 	{
 		// Optionally show the alpha key matte
-		Color.rgb = sRGBToLinear(float3(Color.a, Color.a, Color.a));
+		Color.rgb = Color.a;
 		Color.a = 1.0f;
 	}
 	else
@@ -78,7 +78,6 @@ void MainCS(
 	const float AlphaMin = Params0.z;
 	const float AlphaMax = Params0.w;
 	const float DespillStrength = Params1.x;
-	const float FringeGamma = Params1.y;
 	const float ChromaticityOnly = Params1.z;
 	
 	float4 SourceColor = InputTexture.SampleLevel(InputSampler, UV, 0);
@@ -124,9 +123,6 @@ void MainCS(
 		FinalColor.rgb = BlueDespillAvg(FinalColor.rgb);
 	}
 	
-	// Darken fringe
-	FinalColor.rgb = FinalColor.rgb * pow(Alpha, FringeGamma);
-	
 	const uint2 PixelPos = DispatchThreadId + Output_ViewportMin;
 	RWOutputTexture[PixelPos] = Output(FinalColor);
 }
diff --git a/Engine/Plugins/Experimental/Compositing/Composite/Shaders/Private/CompositeDeltaKeyerCommon.ush b/Engine/Plugins/Experimental/Compositing/Composite/Shaders/Private/CompositeDeltaKeyerCommon.ush
--- a/Engine/Plugins/Experimental/Compositing/Composite/Shaders/Private/CompositeDeltaKeyerCommon.ush
+++ b/Engine/Plugins/Experimental/Compositing/Composite/Shaders/Private/CompositeDeltaKeyerCommon.ush
@@ -158,22 +158,19 @@ float4 BlueScreenKeyer(
 	float  AlphaGain = 1.0f)
 {
 	// Blue suppression: remove blue proportionally to the size of other channels
-	float PlateB = max(Plate.b - RedInfluence * Plate.r
-		- GreenInfluence * Plate.g, 0.0);
-	float ScreenB = max(ScreenColor.b - RedInfluence * ScreenColor.r
-		- GreenInfluence * ScreenColor.g, 0.0);
+	float PlateB = max(Plate.b - RedInfluence * Plate.r - GreenInfluence * Plate.g, 0.0);
+	float ScreenB = max(ScreenColor.b - RedInfluence * ScreenColor.r - GreenInfluence * ScreenColor.g, 0.0);
 
 	// Ratio of how much blue is left on the plate over the screen
 	float BlueRatio = saturate(PlateB * SafeRcp(ScreenB));
-	float3 ScreenAdj = float3(ScreenColor.r, ScreenColor.g, ScreenB);
-	float3 BlueSpill = clamp(ScreenAdj * BlueRatio, 0.0, 32.0);
 
-	// Final color
+	// Despill
+	float3 ScreenAdj = float3(ScreenColor.r, ScreenColor.g, ScreenB);
+	float3 BlueSpill = ScreenAdj * BlueRatio;
 	float3 ResultRGB = clamp(Plate - BlueSpill, 0.0, 1024.0);
 
 	// Alpha matte
-	float LumaDiff = Luminance(Plate - ScreenColor);
-	float Alpha = saturate(max(LumaDiff, 1.0 - BlueRatio) * AlphaGain);
+	float Alpha = saturate(1.0 - BlueRatio) * AlphaGain;
 
 	return float4(ResultRGB, Alpha);
 }
diff --git a/Engine/Plugins/Experimental/Compositing/Composite/Source/Composite/Private/Passes/CompositePassDeltaKeyer.cpp b/Engine/Plugins/Experimental/Compositing/Composite/Source/Composite/Private/Passes/CompositePassDeltaKeyer.cpp
--- a/Engine/Plugins/Experimental/Compositing/Composite/Source/Composite/Private/Passes/CompositePassDeltaKeyer.cpp
+++ b/Engine/Plugins/Experimental/Compositing/Composite/Source/Composite/Private/Passes/CompositePassDeltaKeyer.cpp
@@ -145,7 +145,6 @@ UCompositePassDeltaKeyer::UCompositePassDeltaKeyer(const FObjectInitializer& Obj
 	, ChromaticityOnly(1.0f)
 	, DespillMethod(ECompositeDeltaKeyerDespillMethod::Default)
 	, DespillStrength(1.0f)
-	, FringeGamma(2.2f)
 	, bVisualizeAlpha(false)
 	, bInvertAlpha(false)
 {
@@ -160,7 +159,7 @@ bool UCompositePassDeltaKeyer::GetProxy(const UE::CompositeCore::FPassInputDecl&
 	FCompositePassDeltaKeyerProxy* Proxy = InFrameAllocator.Create<FCompositePassDeltaKeyerProxy>(UE::CompositeCore::FPassInputDeclArray{ InputDecl });
 	Proxy->KeyColor = KeyColor;
 	Proxy->Params0 = FVector4f(ColorTolerance.X, ColorTolerance.Y, AlphaThreshold.X, AlphaThreshold.Y);
-	Proxy->Params1 = FVector4f(DespillStrength, FringeGamma, FMath::Clamp(ChromaticityOnly, 0.0f, 1.0f), 0.0f);
+	Proxy->Params1 = FVector4f(DespillStrength, 0.0f, FMath::Clamp(ChromaticityOnly, 0.0f, 1.0f), 0.0f);
 	Proxy->DespillMethod = static_cast<uint32>(DespillMethod);
 	Proxy->bInvertAlpha = bInvertAlpha;
 	Proxy->bVisualizeAlpha = bVisualizeAlpha;
diff --git a/Engine/Plugins/Experimental/Compositing/Composite/Source/Composite/Public/Passes/CompositePassDeltaKeyer.h b/Engine/Plugins/Experimental/Compositing/Composite/Source/Composite/Public/Passes/CompositePassDeltaKeyer.h
--- a/Engine/Plugins/Experimental/Compositing/Composite/Source/Composite/Public/Passes/CompositePassDeltaKeyer.h
+++ b/Engine/Plugins/Experimental/Compositing/Composite/Source/Composite/Public/Passes/CompositePassDeltaKeyer.h
@@ -63,7 +63,7 @@ public:
 	UPROPERTY(EditAnywhere, BlueprintReadWrite, Category = "", meta = (UIMin = 0.0f, UIMax = 1.0f))
 	FVector2f AlphaThreshold;
 
-	/** Whether or not to include lightness/luminance in the color tolerance calculation. */
+	/** Use chromaticity only, or include luminance in the color tolerance calculation. */
 	UPROPERTY(EditAnywhere, BlueprintReadWrite, Category = "", meta = (UIMin = 0.0f, UIMax = 1.0f))
 	float ChromaticityOnly;
 
@@ -75,10 +75,6 @@ public:
 	UPROPERTY(EditAnywhere, BlueprintReadWrite, Category = "", meta = (UIMin = 0.0f, UIMax = 1.0f))
 	float DespillStrength;
 
-	/** Gamma adjustment for the edge fringe between full translucency & full opacity. */
-	UPROPERTY(EditAnywhere, BlueprintReadWrite, Category = "", meta = (ClampMin = 0.1f, ClampMax = 4.0))
-	float FringeGamma;
-
 	/** Vizualize the alpha key. */
 	UPROPERTY(EditAnywhere, Transient, Category = "")
 	bool bVisualizeAlpha;
