Composite: Removed color distance approach for now due to poor edge quality.

* Added fill & no-premult visualization mode.
* Removed despill method selection as Luv chromaticity approach works best thus far.

diff --git a/Engine/Plugins/Experimental/Compositing/Composite/Shaders/Private/CompositeDeltaKeyer.usf b/Engine/Plugins/Experimental/Compositing/Composite/Shaders/Private/CompositeDeltaKeyer.usf
--- a/Engine/Plugins/Experimental/Compositing/Composite/Shaders/Private/CompositeDeltaKeyer.usf
+++ b/Engine/Plugins/Experimental/Compositing/Composite/Shaders/Private/CompositeDeltaKeyer.usf
@@ -10,14 +10,6 @@
     #error TILE_SIZE must be defined
 #endif
 
-// Hand-picked scale factor used to adjust LUV color delta range
-static const float kDistanceScaleFactor_CIELUV = 2.0 * DeltaE_JND / Lstar_CIELUV;
-
-#define DESPILL_METHOD_NONE 0
-#define DESPILL_METHOD_DEFAULT 1
-#define DESPILL_METHOD_GREEN_AVG 2
-#define DESPILL_METHOD_BLUE_AVG 3
-
 Texture2D InputTexture;
 SamplerState InputSampler;
 
@@ -32,8 +24,8 @@ SCREEN_PASS_TEXTURE_VIEWPORT(Output)
 float3 KeyColor;
 float4 Params0;
 float4 Params1;
-uint DespillMethod;
-uint bVisualizeAlpha;
+uint BackingChannel;
+uint Visualization;
 uint bInvertAlpha;
 
 // Output color function with debug convenience
@@ -42,15 +34,24 @@ float4 Output(float4 Color)
 	// Optionally invert alpha
 	Color.a = select(bInvertAlpha > 0, 1.0f - Color.a, Color.a);
 
-	if(bVisualizeAlpha > 0)
+	if(Visualization == 1)
 	{
 		// Optionally show the alpha key matte
-		Color.rgb = Color.a;
+		Color.rgb = sRGBToLinear(Color.a);
+		Color.a = 1.0f;
+	}
+	else if(Visualization == 2)
+	{
+		// Opaque alpha to show fill
 		Color.a = 1.0f;
 	}
+	else if(Visualization == 3)
+	{
+		// no-premult
+	}
 	else
 	{
-		// Return premultiplied alpha
+		// Return premultiplied alpha for composite marge passes
 		Color.rgb *= Color.a;
 	}
 
@@ -73,56 +74,55 @@ void MainCS(
 	}
 
 	// Input parameters
-	const float DeltaThresholdMin = Params0.x;
-	const float DeltaThresholdMax = Params0.y;
-	const float AlphaMin = Params0.z;
-	const float AlphaMax = Params0.w;
-	const float DespillStrength = Params1.x;
-	const float ChromaticityOnly = Params1.z;
+	const float3 ChannelWeights = Params0.xyz;
+	const float DespillStrength = Params0.w;
+	const float AlphaMin = Params1.x;
+	const float AlphaMax = Params1.y;
 	
-	float4 SourceColor = InputTexture.SampleLevel(InputSampler, UV, 0);
+	float4 Foreground = InputTexture.SampleLevel(InputSampler, UV, 0);
 #if USE_CLEAN_PLATE
 	float3 Background = CleanPlateTexture.SampleLevel(CleanPlateSampler, UV, 0).rgb;
 #else
 	float3 Background = KeyColor;
 #endif
 
-	// Map the source color onto what we call the "foreground".
-	// Convert both to Luv color space
-	float3 Foreground_Luv = RGB_to_Luv(SourceColor.rgb);
-	float3 Background_Luv = RGB_to_Luv(Background);
-
-	// Optionally preserve or discard the lightness contribution by shrinking its component distance
-	Background_Luv[0] = lerp(Background_Luv[0], Foreground_Luv[0], ChromaticityOnly);
-	
-    // Calculate chromaticity difference & alpha
-    float Delta = kDistanceScaleFactor_CIELUV * distance(Foreground_Luv, Background_Luv);
-    float Alpha = Linearstep(DeltaThresholdMin, DeltaThresholdMax, Delta);
-	
-	// Project the foreground color onto the background in uv-space
-	float2 ProjectedDiff = ProjectClamped(Foreground_Luv.yz, Background_Luv.yz);
+	float4 Result = Foreground;
 
-	if(DespillMethod == DESPILL_METHOD_DEFAULT)
+	BRANCH
+	if(BackingChannel == 0)
 	{
-		// Despill the foreground
-		Foreground_Luv.yz -= DespillStrength * ProjectedDiff;
+		Result = RedScreenKeyer(Foreground.rgb, Background, ChannelWeights[1], ChannelWeights[2]);
 	}
-
-	// Apply alpha threshold range
-	float AdjustedAlpha = Linearstep(AlphaMin, AlphaMax, Alpha);
-
-    // Convert back to RGB
-	float4 FinalColor = float4(Luv_to_RGB(Foreground_Luv), AdjustedAlpha);
-
-	if(DespillMethod == DESPILL_METHOD_GREEN_AVG)
+	else if(BackingChannel == 1)
 	{
-		FinalColor.rgb = GreenDespillAvg(FinalColor.rgb);
+		Result = GreenScreenKeyer(Foreground.rgb, Background, ChannelWeights[0], ChannelWeights[2]);
 	}
-	else if(DespillMethod == DESPILL_METHOD_BLUE_AVG)
+	else if(BackingChannel == 2)
+	{
+		Result = BlueScreenKeyer(Foreground.rgb, Background, ChannelWeights[0], ChannelWeights[1]);
+	}
+
+	BRANCH
+	if(DespillStrength > 0.0f)
 	{
-		FinalColor.rgb = BlueDespillAvg(FinalColor.rgb);
+		// Map the source color onto what we call the "foreground".
+		// Convert both to Luv color space
+		float3 Foreground_Luv = RGB_to_Luv(Result.rgb);
+		float3 Background_Luv = RGB_to_Luv(Background);
+
+		// Project the foreground color onto the background in uv-space
+		float2 ProjectedDiff = ProjectClamped(Foreground_Luv.yz, Background_Luv.yz);
+
+		// Despill the foreground
+		Foreground_Luv.yz -= DespillStrength * ProjectedDiff;
+
+    	// Convert back to RGB
+		Result = float4(Luv_to_RGB(Foreground_Luv), Result.a);
 	}
+
+	// Apply final alpha threshold
+	Result.a = Linearstep(AlphaMin, AlphaMax, Result.a);
 	
 	const uint2 PixelPos = DispatchThreadId + Output_ViewportMin;
-	RWOutputTexture[PixelPos] = Output(FinalColor);
+	RWOutputTexture[PixelPos] = Output(Result);
 }
diff --git a/Engine/Plugins/Experimental/Compositing/Composite/Shaders/Private/CompositeDeltaKeyerCommon.ush b/Engine/Plugins/Experimental/Compositing/Composite/Shaders/Private/CompositeDeltaKeyerCommon.ush
--- a/Engine/Plugins/Experimental/Compositing/Composite/Shaders/Private/CompositeDeltaKeyerCommon.ush
+++ b/Engine/Plugins/Experimental/Compositing/Composite/Shaders/Private/CompositeDeltaKeyerCommon.ush
@@ -8,7 +8,6 @@
 
 // CIE 1931 D65 white point
 static const float3 XYZ_n = xyY_2_XYZ(float3(0.3127, 0.3290, 1.0));
-
 static const float ZeroEpsilon = 1.0e-6;
 static const float DeltaE_JND = 1.0;
 static const float Lstar_CIELUV = 100.0;
@@ -33,27 +32,27 @@ float Linearstep(float MinValue, float MaxValue, float X)
 
 float2 Project(float2 VecA, float2 VecB)
 {
-	float Dot = max(0.0f, dot(VecA, VecB));
-	float LengthSquaredB = dot(VecB, VecB);
+	float AdotB = dot(VecA, VecB);
+	float BdotB = dot(VecB, VecB);
 	
-	return VecB * (Dot / LengthSquaredB);
+	return select(BdotB < ZeroEpsilon, float2(0, 0), VecB * (AdotB / BdotB));
 }
 
 // Project vector A onto B, and return 0 if the projection was negative
 float2 ProjectClamped(float2 VecA, float2 VecB)
 {
-	float Dot = max(0.0f, dot(VecA, VecB));
-	float LengthSquaredB = dot(VecB, VecB);
+	float AdotB = max(0.0f, dot(VecA, VecB));
+	float BdotB = dot(VecB, VecB);
 	
-	return select(LengthSquaredB < ZeroEpsilon, float2(0, 0), VecB * (Dot / LengthSquaredB));
+	return select(BdotB < ZeroEpsilon, float2(0, 0), VecB * (AdotB / BdotB));
 }
 
 float3 ProjectClamped(float3 VecA, float3 VecB)
 {
-	float Dot = max(0.0f, dot(VecA, VecB));
-	float LengthSquaredB = dot(VecB, VecB);
+	float AdotB = max(0.0f, dot(VecA, VecB));
+	float BdotB = dot(VecB, VecB);
 
-	return select(LengthSquaredB < ZeroEpsilon, float3(0, 0, 0), VecB * (Dot / LengthSquaredB));
+	return select(BdotB < ZeroEpsilon, float3(0, 0, 0), VecB * (AdotB / BdotB));
 }
 
 // Convert color from XYZ to CIE 1976 L*, u*, v* color space.
@@ -149,45 +148,50 @@ float3 BlueDespillAvg(float3 Color)
 	);
 }
 
-// --- Unused but kept for testing ---
-float4 BlueScreenKeyer(
-	float3 Plate,
-	float3 ScreenColor,
+float4 GreenScreenKeyer(
+	float3 Foreground,
+	float3 Background,
 	float  RedInfluence,
-	float  GreenInfluence,
-	float  AlphaGain = 1.0f)
+	float  BlueInfluence)
 {
-	// Blue suppression: remove blue proportionally to the size of other channels
-	float PlateB = max(Plate.b - RedInfluence * Plate.r - GreenInfluence * Plate.g, 0.0);
-	float ScreenB = max(ScreenColor.b - RedInfluence * ScreenColor.r - GreenInfluence * ScreenColor.g, 0.0);
-
-	// Ratio of how much blue is left on the plate over the screen
-	float BlueRatio = saturate(PlateB * SafeRcp(ScreenB));
+	float ForegroundG = max(Foreground.g - RedInfluence * Foreground.r - BlueInfluence * Foreground.b, 0.0);
+	float BackgroundG = max(Background.g - RedInfluence * Background.r - BlueInfluence * Background.b, 0.0);
+	float GreenRatio  = saturate(ForegroundG * SafeRcp(BackgroundG));
 
-	// Despill
-	float3 ScreenAdj = float3(ScreenColor.r, ScreenColor.g, ScreenB);
-	float3 BlueSpill = ScreenAdj * BlueRatio;
-	float3 ResultRGB = clamp(Plate - BlueSpill, 0.0, 1024.0);
+	float3 BackgroundAdj  = float3(Background.r, BackgroundG, Background.b);
+	float3 GreenSpill     = BackgroundAdj * GreenRatio;
+	float3 Fill           = max(Foreground - GreenSpill, 0.0);
+	float Alpha           = saturate(1.0 - GreenRatio);
 
-	// Alpha matte
-	float Alpha = saturate(1.0 - BlueRatio) * AlphaGain;
+	return float4(Fill, Alpha);
+}
 
-	return float4(ResultRGB, Alpha);
+float3 SwizzleRG(float3 InColor)
+{
+	return float3(InColor.g, InColor.r, InColor.b);
 }
 
-float3 SwizzleGreenBlue(float3 InColor)
+float3 SwizzleGB(float3 InColor)
 {
 	return float3(InColor.r, InColor.b, InColor.g);
 }
 
-float4 GreenScreenKeyer(float3 Plate, float3 ScreenColor, float RedInfluence, float BlueInfluence, float AlphaGain = 1.0f)
+float4 RedScreenKeyer(float3 Foreground, float3 Background, float GreenInfluence, float BlueInfluence)
 {
-	float3 SwizzledPlate = SwizzleGreenBlue(Plate);
-	float3 SizzledScreenColor = SwizzleGreenBlue(ScreenColor);
-	
-	float4 Output = BlueScreenKeyer(SwizzledPlate, SizzledScreenColor, RedInfluence, BlueInfluence, AlphaGain);
-	Output.rgb = SwizzleGreenBlue(Output.rgb);
+	float3 SwizzledForeground = SwizzleRG(Foreground);
+	float3 SizzledBackground = SwizzleRG(Background);
+	float4 Output = GreenScreenKeyer(SwizzledForeground, SizzledBackground, GreenInfluence, BlueInfluence);
+	Output.rgb = SwizzleRG(Output.rgb);
 
 	return Output;
 }
-// --- End Unused ---
\ No newline at end of file
+
+float4 BlueScreenKeyer(float3 Foreground, float3 Background, float RedInfluence, float GreenInfluence)
+{
+	float3 SwizzledForeground = SwizzleGB(Foreground);
+	float3 SizzledBackground = SwizzleGB(Background);
+	float4 Output = GreenScreenKeyer(SwizzledForeground, SizzledBackground, RedInfluence, GreenInfluence);
+	Output.rgb = SwizzleGB(Output.rgb);
+
+	return Output;
+}
\ No newline at end of file
diff --git a/Engine/Plugins/Experimental/Compositing/Composite/Source/Composite/Private/Passes/CompositePassDeltaKeyer.cpp b/Engine/Plugins/Experimental/Compositing/Composite/Source/Composite/Private/Passes/CompositePassDeltaKeyer.cpp
--- a/Engine/Plugins/Experimental/Compositing/Composite/Source/Composite/Private/Passes/CompositePassDeltaKeyer.cpp
+++ b/Engine/Plugins/Experimental/Compositing/Composite/Source/Composite/Private/Passes/CompositePassDeltaKeyer.cpp
@@ -33,9 +33,9 @@ class FCompositePassDeltaKeyerShader : public FGlobalShader
 		SHADER_PARAMETER(FVector3f, KeyColor)
 		SHADER_PARAMETER(FVector4f, Params0)
 		SHADER_PARAMETER(FVector4f, Params1)
-		SHADER_PARAMETER(uint32, bVisualizeAlpha)
+		SHADER_PARAMETER(uint32, BackingChannel)
+		SHADER_PARAMETER(uint32, Visualization)
 		SHADER_PARAMETER(uint32, bInvertAlpha)
-		SHADER_PARAMETER(uint32, DespillMethod)
 	END_SHADER_PARAMETER_STRUCT()
 
 	static inline void ModifyCompilationEnvironment(const FGlobalShaderPermutationParameters& Parameters, FShaderCompilerEnvironment& OutEnvironment)
@@ -92,8 +92,8 @@ namespace UE
 				PassParameters->KeyColor = FVector3f(KeyColor);
 				PassParameters->Params0 = Params0;
 				PassParameters->Params1 = Params1;
-				PassParameters->DespillMethod = DespillMethod;
-				PassParameters->bVisualizeAlpha = static_cast<uint32>(bVisualizeAlpha);
+				PassParameters->BackingChannel = BackingChannel;
+				PassParameters->Visualization = static_cast<uint32>(Visualization);
 				PassParameters->bInvertAlpha = static_cast<uint32>(bInvertAlpha);
 	
 				bool bUseCleanPlate = false;
@@ -125,13 +125,12 @@ namespace UE
 				return UE::CompositeCore::FPassTexture{ MoveTemp(Output), Metadata };
 			}
 
-			FLinearColor KeyColor = FLinearColor::White;
-			FVector4f Params0 = FVector4f::Zero();
-			FVector4f Params1 = FVector4f::Zero();
-			uint32 DespillMethod = 0;
-			FLinearColor DespillColor = FLinearColor::White;
-			bool bVisualizeAlpha = false;
-			bool bInvertAlpha = false;
+			uint32 BackingChannel;
+			FLinearColor KeyColor;
+			FVector4f Params0;
+			FVector4f Params1;
+			ECompositeDeltaKeyerVisualization Visualization;
+			bool bInvertAlpha;
 			TWeakObjectPtr<UTexture> CleanPlateWeakPtr;
 		};
 	}
@@ -140,14 +139,12 @@ namespace UE
 UCompositePassDeltaKeyer::UCompositePassDeltaKeyer(const FObjectInitializer& ObjectInitializer)
 	: Super(ObjectInitializer)
 	, KeyColor(FLinearColor::Green)
-	, ColorTolerance(0.08f, 1.0f)
+	, ChannelWeights(0.5f, 0.5f, 0.5f)
 	, AlphaThreshold(0.0f, 1.0f)
-	, ChromaticityOnly(1.0f)
-	, DespillMethod(ECompositeDeltaKeyerDespillMethod::Default)
 	, DespillStrength(1.0f)
-	, bVisualizeAlpha(false)
+	, Visualization(ECompositeDeltaKeyerVisualization::None)
 	, bInvertAlpha(false)
-{
+{ 
 }
 
 UCompositePassDeltaKeyer::~UCompositePassDeltaKeyer() = default;
@@ -158,11 +155,11 @@ bool UCompositePassDeltaKeyer::GetProxy(const UE::CompositeCore::FPassInputDecl&
 
 	FCompositePassDeltaKeyerProxy* Proxy = InFrameAllocator.Create<FCompositePassDeltaKeyerProxy>(UE::CompositeCore::FPassInputDeclArray{ InputDecl });
 	Proxy->KeyColor = KeyColor;
-	Proxy->Params0 = FVector4f(ColorTolerance.X, ColorTolerance.Y, AlphaThreshold.X, AlphaThreshold.Y);
-	Proxy->Params1 = FVector4f(DespillStrength, 0.0f, FMath::Clamp(ChromaticityOnly, 0.0f, 1.0f), 0.0f);
-	Proxy->DespillMethod = static_cast<uint32>(DespillMethod);
+	Proxy->Params0 = FVector4f(ChannelWeights.X, ChannelWeights.Y, ChannelWeights.Z, DespillStrength);
+	Proxy->Params1 = FVector4f(AlphaThreshold.X, AlphaThreshold.Y, 0.0f, 0.0f);
+	Proxy->BackingChannel = static_cast<uint32>(FMath::Max3Index(KeyColor.R, KeyColor.G, KeyColor.B));
 	Proxy->bInvertAlpha = bInvertAlpha;
-	Proxy->bVisualizeAlpha = bVisualizeAlpha;
+	Proxy->Visualization = Visualization;
 	Proxy->CleanPlateWeakPtr = CleanPlate;
 
 	OutProxy = Proxy;
diff --git a/Engine/Plugins/Experimental/Compositing/Composite/Source/Composite/Public/Passes/CompositePassDeltaKeyer.h b/Engine/Plugins/Experimental/Compositing/Composite/Source/Composite/Public/Passes/CompositePassDeltaKeyer.h
--- a/Engine/Plugins/Experimental/Compositing/Composite/Source/Composite/Public/Passes/CompositePassDeltaKeyer.h
+++ b/Engine/Plugins/Experimental/Compositing/Composite/Source/Composite/Public/Passes/CompositePassDeltaKeyer.h
@@ -9,17 +9,16 @@
 #define UE_API COMPOSITE_API
 
 UENUM(BlueprintType)
-enum class ECompositeDeltaKeyerDespillMethod : uint8
+enum class ECompositeDeltaKeyerVisualization : uint8
 {
 	None = 0,
-	Default = 1,
-	GreenAverage = 2,
-	BlueAverage = 3,
+	Alpha = 1,
+	Fill = 2,
+	NoPremult = 3,
 };
 
 /**
- * Chromaticity color keyer based on color distances in (L)UV color space, with spill removal support.
- * Assumes input is in linear working color space.
+ * Color keyer with clean plate & spill removal support.
 */
 UCLASS(MinimalAPI, BlueprintType, Blueprintable, EditInlineNew, CollapseCategories, meta=(DisplayName="Delta Keyer Pass"))
 class UCompositePassDeltaKeyer : public UCompositePassBase
@@ -41,43 +40,32 @@ public:
 #endif
 
 public:
-	/** Background key color. */
-	UPROPERTY(EditAnywhere, BlueprintReadWrite, Category = "", meta = (OnlyUpdateOnInteractionEnd, EditCondition = "CleanPlate == nullptr"))
+	/** Background key color. Always set as we currently determine R/G/B backing type from its maximum channel. */
+	UPROPERTY(EditAnywhere, BlueprintReadWrite, Category = "", meta = (OnlyUpdateOnInteractionEnd))
 	FLinearColor KeyColor;
 
 	/**
-	 * Clean plate background for calculating color differences, replacing the static key color.
+	 * Clean plate background for calculating color differences, refining the static key color.
 	 * Resolution must match the composite actor render resolution.
 	 */
 	UPROPERTY(EditAnywhere, BlueprintReadWrite, Category = "")
 	TObjectPtr<UTexture> CleanPlate;
 
-	/**
-	 * A color tolerance range for keying, measured as the distance between the input & key color.
-	 * The minimum tolerance will be fully keyed, while the maximum fully opaque.
-	 */
-	UPROPERTY(EditAnywhere, BlueprintReadWrite, Category = "", meta = (UIMin = 0.0f, UIMax = 2.0f))
-	FVector2f ColorTolerance;
+	/** Weights of the non-backing channels affecting the keying matte hardness. */
+	UPROPERTY(EditAnywhere, BlueprintReadWrite, Category = "", meta = (UIMin = 0.0f, UIMax = 1.5f))
+	FVector3f ChannelWeights;
 
 	/** Thresholds any alpha value outside the specified range to zero or one respectively, with linear interpolation in-between. */
 	UPROPERTY(EditAnywhere, BlueprintReadWrite, Category = "", meta = (UIMin = 0.0f, UIMax = 1.0f))
 	FVector2f AlphaThreshold;
 
-	/** Use chromaticity only, or include luminance in the color tolerance calculation. */
-	UPROPERTY(EditAnywhere, BlueprintReadWrite, Category = "", meta = (UIMin = 0.0f, UIMax = 1.0f))
-	float ChromaticityOnly;
-
-	/** Despill method, to remove/neutralize the residual key color spill on the subject. */
-	UPROPERTY(EditAnywhere, BlueprintReadWrite, Category = "")
-	ECompositeDeltaKeyerDespillMethod DespillMethod;
-
 	/** Strength of the spill reduction, 0.0: none, 1.0: full. */
 	UPROPERTY(EditAnywhere, BlueprintReadWrite, Category = "", meta = (UIMin = 0.0f, UIMax = 1.0f))
 	float DespillStrength;
 
-	/** Vizualize the alpha key. */
-	UPROPERTY(EditAnywhere, Transient, Category = "")
-	bool bVisualizeAlpha;
+	/** Vizualize the alpha key or fill. */
+	UPROPERTY(EditAnywhere, BlueprintReadWrite, Category = "")
+	ECompositeDeltaKeyerVisualization Visualization;
 
 	/** Invert the alpha key. */
 	UPROPERTY(EditAnywhere, BlueprintReadWrite, Category = "")
