Composite: Fix keyer pass output resolution handling.

diff --git a/Engine/Plugins/Experimental/Compositing/Composite/Shaders/Private/CompositeDeltaKeyer.usf b/Engine/Plugins/Experimental/Compositing/Composite/Shaders/Private/CompositeDeltaKeyer.usf
--- a/Engine/Plugins/Experimental/Compositing/Composite/Shaders/Private/CompositeDeltaKeyer.usf
+++ b/Engine/Plugins/Experimental/Compositing/Composite/Shaders/Private/CompositeDeltaKeyer.usf
@@ -65,7 +65,7 @@ void MainCS(
     uint GroupThreadIndex : SV_GroupIndex,
     uint2 DispatchThreadId : SV_DispatchThreadID)
 {
-	float4 SvPosition = float4((float2)DispatchThreadId + Input_ViewportMin + 0.5f, 0.0f, 1.0f);
+	float4 SvPosition = float4((float2)DispatchThreadId + Output_ViewportMin + 0.5f, 0.0f, 1.0f);
 	float2 UV = (SvPosition.xy) * Output_ExtentInverse;
 
 	if (IsComputeUVOutOfBounds(UV))
diff --git a/Engine/Plugins/Experimental/Compositing/Composite/Source/Composite/Private/Passes/CompositePassDeltaKeyer.cpp b/Engine/Plugins/Experimental/Compositing/Composite/Source/Composite/Private/Passes/CompositePassDeltaKeyer.cpp
--- a/Engine/Plugins/Experimental/Compositing/Composite/Source/Composite/Private/Passes/CompositePassDeltaKeyer.cpp
+++ b/Engine/Plugins/Experimental/Compositing/Composite/Source/Composite/Private/Passes/CompositePassDeltaKeyer.cpp
@@ -77,7 +77,6 @@ namespace UE
 				}
 
 				FGlobalShaderMap* GlobalShaderMap = GetGlobalShaderMap(InView.GetFeatureLevel());
-				const FIntPoint TextureSize = Input.Texture->Desc.Extent;
 
 				const bool bSameSize = Output.ViewRect.Size() == Input.ViewRect.Size();
 				const FScreenPassTextureViewportParameters InputParameters = GetScreenPassTextureViewportParameters(FScreenPassTextureViewport(Input));
@@ -112,14 +111,15 @@ namespace UE
 				FCompositePassDeltaKeyerShader::FPermutationDomain PermutationVector;
 				PermutationVector.Set<FCompositePassDeltaKeyerShader::FUseCleanPlate>(bUseCleanPlate);
 
+				const FIntPoint OutputSize = Output.ViewRect.Size();
 				TShaderMapRef<FCompositePassDeltaKeyerShader> ComputeShader(GlobalShaderMap, PermutationVector);
 				FComputeShaderUtils::AddPass(
 					GraphBuilder,
-					RDG_EVENT_NAME("Composite.DeltaKeyer (%dx%d)", TextureSize.X, TextureSize.Y),
+					RDG_EVENT_NAME("Composite.DeltaKeyer (%dx%d)", OutputSize.X, OutputSize.Y),
 					ERDGPassFlags::Compute,
 					ComputeShader,
 					PassParameters,
-					FComputeShaderUtils::GetGroupCount(TextureSize, FCompositePassDeltaKeyerShader::ThreadGroupSize)
+					FComputeShaderUtils::GetGroupCount(OutputSize, FCompositePassDeltaKeyerShader::ThreadGroupSize)
 				);
 
 				return UE::CompositeCore::FPassTexture{ MoveTemp(Output), Metadata };
