Composite: Minor shader cleanup.

diff --git a/Engine/Plugins/Experimental/Compositing/Composite/Shaders/Private/CompositeDeltaKeyer.usf b/Engine/Plugins/Experimental/Compositing/Composite/Shaders/Private/CompositeDeltaKeyer.usf
--- a/Engine/Plugins/Experimental/Compositing/Composite/Shaders/Private/CompositeDeltaKeyer.usf
+++ b/Engine/Plugins/Experimental/Compositing/Composite/Shaders/Private/CompositeDeltaKeyer.usf
@@ -117,7 +117,7 @@ void MainCS(
 		Foreground_Luv.yz -= DespillStrength * ProjectedDiff;
 
     	// Convert back to RGB
-		Result = float4(Luv_to_RGB(Foreground_Luv), Result.a);
+		Result.rgb = Luv_to_RGB(Foreground_Luv);
 	}
 
 	// Apply final alpha threshold
diff --git a/Engine/Plugins/Experimental/Compositing/Composite/Shaders/Private/CompositeDeltaKeyerCommon.ush b/Engine/Plugins/Experimental/Compositing/Composite/Shaders/Private/CompositeDeltaKeyerCommon.ush
--- a/Engine/Plugins/Experimental/Compositing/Composite/Shaders/Private/CompositeDeltaKeyerCommon.ush
+++ b/Engine/Plugins/Experimental/Compositing/Composite/Shaders/Private/CompositeDeltaKeyerCommon.ush
@@ -25,9 +25,9 @@ float LuminanceSRGB(float3 InColor)
 }
 
 // Linear equivalent of smoothstep
-float Linearstep(float MinValue, float MaxValue, float X)
+float Linearstep(float Edge0, float Edge1, float X)
 {
-	return saturate((X - MinValue) / (MaxValue - MinValue));
+	return saturate((X - Edge0) / (Edge1 - Edge0));
 }
 
 float2 Project(float2 VecA, float2 VecB)
