Composite: Keyer improvements.

* Prevent double premult since the fixed keyer output is already considered premult.
* Added unused experiments in Luv space.

diff --git a/Engine/Plugins/Experimental/Compositing/Composite/Shaders/Private/CompositeDeltaKeyer.usf b/Engine/Plugins/Experimental/Compositing/Composite/Shaders/Private/CompositeDeltaKeyer.usf
--- a/Engine/Plugins/Experimental/Compositing/Composite/Shaders/Private/CompositeDeltaKeyer.usf
+++ b/Engine/Plugins/Experimental/Compositing/Composite/Shaders/Private/CompositeDeltaKeyer.usf
@@ -45,15 +45,6 @@ float4 Output(float4 Color)
 		// Opaque alpha to show fill
 		Color.a = 1.0f;
 	}
-	else if(Visualization == 3)
-	{
-		// no-premult
-	}
-	else
-	{
-		// Return premultiplied alpha for composite marge passes
-		Color.rgb *= Color.a;
-	}
 
 	return Color;
 }
diff --git a/Engine/Plugins/Experimental/Compositing/Composite/Shaders/Private/CompositeDeltaKeyerCommon.ush b/Engine/Plugins/Experimental/Compositing/Composite/Shaders/Private/CompositeDeltaKeyerCommon.ush
--- a/Engine/Plugins/Experimental/Compositing/Composite/Shaders/Private/CompositeDeltaKeyerCommon.ush
+++ b/Engine/Plugins/Experimental/Compositing/Composite/Shaders/Private/CompositeDeltaKeyerCommon.ush
@@ -55,6 +55,19 @@ float3 ProjectClamped(float3 VecA, float3 VecB)
 	return select(BdotB < ZeroEpsilon, float3(0, 0, 0), VecB * (AdotB / BdotB));
 }
 
+float GetAngle(float2 A, float2 B)
+{
+	return acos(dot(normalize(A), normalize(B)));
+}
+
+float2 Rotate2D(float2 V, float Theta)
+{
+	float S, C;
+	sincos(Theta, S, C);
+	return float2(C*V.x - S*V.y,
+				  S*V.x + C*V.y);
+}
+
 // Convert color from XYZ to CIE 1976 L*, u*, v* color space.
 // We use the latter for its improved perceptual uniformity, as a way to measure chromaticity distance for keying.
 float3 CIE_XYZ_to_Luv(float3 XYZ)
@@ -151,18 +164,48 @@ float3 BlueDespillAvg(float3 Color)
 float4 GreenScreenKeyer(
 	float3 Foreground,
 	float3 Background,
-	float  RedInfluence,
-	float  BlueInfluence)
+	float  RedWeight,
+	float  BlueWeight)
 {
-	float ForegroundG = max(Foreground.g - RedInfluence * Foreground.r - BlueInfluence * Foreground.b, 0.0);
-	float BackgroundG = max(Background.g - RedInfluence * Background.r - BlueInfluence * Background.b, 0.0);
-	float GreenRatio  = saturate(ForegroundG * SafeRcp(BackgroundG));
+	// Shared weighted sums
+	float wF = RedWeight * Foreground.r + BlueWeight * Foreground.b;
+	float wB = RedWeight * Background.r + BlueWeight * Background.b;
+
+	// Green minus weighted red/blue
+	float Fg = Foreground.g - wF;
+	float Bg = Background.g - wB;
+
+	// Ratios
+	float InvBg       = SafeRcp(Bg);
+	float GreenRatio  = saturate(Fg * InvBg);
+	float Alpha       = 1 - GreenRatio;
+
+	// Built-in "despill"
+	float3 DespillF = float3(Foreground.r, wF, Foreground.b);
+	float3 DespillB = float3(Background.r, wB, Background.b);
+
+	// Masked background
+	float3 MaskedB = DespillB * GreenRatio;
+
+	// Handle screen brighter than sample color
+	float3 Delta = Foreground - Background;
+	float Mg     = Delta.g - (RedWeight * Delta.r + BlueWeight * Delta.b);
+	MaskedB.g   += saturate(Mg * InvBg * wB);
+
+	// Fill
+	float3 Fill = DespillF - MaskedB;
+
+	// Re-insert removed green
+	Fill.g += min(Fg, 0.0);
+	// Prevent negativtes
+	Fill = max(Fill, 0.0);
 
-	float3 BackgroundAdj  = float3(Background.r, BackgroundG, Background.b);
-	float3 GreenSpill     = BackgroundAdj * GreenRatio;
-	float3 Fill           = max(Foreground - GreenSpill, 0.0);
-	float Alpha           = saturate(1.0 - GreenRatio);
+	// Hack to make sure the bright areas of the plate are solid
+	Alpha = max(Alpha, Fill.r);
+	// Clamp alpha
+	Alpha = saturate(Alpha);
 
+	// Note: Output is already considered pre-multiplied
 	return float4(Fill, Alpha);
 }
 
@@ -176,22 +219,56 @@ float3 SwizzleGB(float3 InColor)
 	return float3(InColor.r, InColor.b, InColor.g);
 }
 
-float4 RedScreenKeyer(float3 Foreground, float3 Background, float GreenInfluence, float BlueInfluence)
+float4 RedScreenKeyer(float3 Foreground, float3 Background, float GreenWeight, float BlueWeight)
 {
 	float3 SwizzledForeground = SwizzleRG(Foreground);
 	float3 SizzledBackground = SwizzleRG(Background);
-	float4 Output = GreenScreenKeyer(SwizzledForeground, SizzledBackground, GreenInfluence, BlueInfluence);
+	float4 Output = GreenScreenKeyer(SwizzledForeground, SizzledBackground, GreenWeight, BlueWeight);
 	Output.rgb = SwizzleRG(Output.rgb);
 
 	return Output;
 }
 
-float4 BlueScreenKeyer(float3 Foreground, float3 Background, float RedInfluence, float GreenInfluence)
+float4 BlueScreenKeyer(float3 Foreground, float3 Background, float RedWeight, float GreenWeight)
 {
 	float3 SwizzledForeground = SwizzleGB(Foreground);
 	float3 SizzledBackground = SwizzleGB(Background);
-	float4 Output = GreenScreenKeyer(SwizzledForeground, SizzledBackground, RedInfluence, GreenInfluence);
+	float4 Output = GreenScreenKeyer(SwizzledForeground, SizzledBackground, RedWeight, GreenWeight);
 	Output.rgb = SwizzleGB(Output.rgb);
 
 	return Output;
+}
+
+// Experimental
+float4 LuvDistanceKeyer(float3 Foreground_Luv, float3 Background_Luv, float MinDistance, float MaxDistance)
+{
+	const float kDistanceScaleFactor_CIELUV = 2 * DeltaE_JND / Lstar_CIELUV;
+	float Delta = distance(Foreground_Luv.yz, Background_Luv.yz) * kDistanceScaleFactor_CIELUV;
+
+	float Alpha = Linearstep(MinDistance, MaxDistance, Delta);
+
+	// TODO: Modulate Foreground_Luv.x
+	return float4(Foreground_Luv, Alpha);
+}
+
+// Experimental
+float4 LuvAngleKeyer(float3 Foreground_Luv, float3 Background_Luv, float Acceptance, float Scale)
+{
+	// Similar to approach described in Video Demystified  - A Handbook for the Digital Engineer
+	float AcceptanceAngle = 2 * PI * Acceptance;
+	
+	// Angle between the background and the "u" axis.
+	float AxisAngle = atan2(Background_Luv.z, Background_Luv.y);
+	
+	// Axis-aligned foreground uv
+	float2 Foreground_uvaa = Rotate2D(Foreground_Luv.yz, AxisAngle);
+
+	float X = Foreground_uvaa.x;
+	float Z_sub = abs(Foreground_uvaa.y) / tan(0.5 * AcceptanceAngle);
+	float Key = select(X < Z_sub, 0.0f, X - Z_sub);
+	
+	float Alpha = saturate(1.0 - Lstar_CIELUV_Inv * Key * Scale);
+
+	// TODO: Modulate Foreground_Luv.x
+	return float4(Foreground_Luv, Alpha);
 }
\ No newline at end of file
diff --git a/Engine/Plugins/Experimental/Compositing/Composite/Source/Composite/Public/Passes/CompositePassDeltaKeyer.h b/Engine/Plugins/Experimental/Compositing/Composite/Source/Composite/Public/Passes/CompositePassDeltaKeyer.h
--- a/Engine/Plugins/Experimental/Compositing/Composite/Source/Composite/Public/Passes/CompositePassDeltaKeyer.h
+++ b/Engine/Plugins/Experimental/Compositing/Composite/Source/Composite/Public/Passes/CompositePassDeltaKeyer.h
@@ -14,7 +14,6 @@ enum class ECompositeDeltaKeyerVisualization : uint8
 	None = 0,
 	Alpha = 1,
 	Fill = 2,
-	NoPremult = 3,
 };
 
 /**
