Composite: Delta keyer property cleanup.

* Added mandatory screen color type setting.
* Separated channel weights into three separate sliders.

diff --git a/Engine/Plugins/Experimental/Compositing/Composite/Shaders/Private/CompositeDeltaKeyer.usf b/Engine/Plugins/Experimental/Compositing/Composite/Shaders/Private/CompositeDeltaKeyer.usf
--- a/Engine/Plugins/Experimental/Compositing/Composite/Shaders/Private/CompositeDeltaKeyer.usf
+++ b/Engine/Plugins/Experimental/Compositing/Composite/Shaders/Private/CompositeDeltaKeyer.usf
@@ -24,7 +24,7 @@ SCREEN_PASS_TEXTURE_VIEWPORT(Output)
 float3 KeyColor;
 float4 Params0;
 float4 Params1;
-uint BackingChannel;
+uint ScreenType;
 uint Visualization;
 uint bInvertAlpha;
 
@@ -36,13 +36,13 @@ float4 Output(float4 Color)
 
 	if(Visualization == 1)
 	{
-		// Optionally show the alpha key matte
-		Color.rgb = sRGBToLinear(Color.a);
+		// Opaque to show fill
 		Color.a = 1.0f;
 	}
 	else if(Visualization == 2)
 	{
-		// Opaque alpha to show fill
+		// Optionally show the alpha key matte
+		Color.rgb = sRGBToLinear(Color.a);
 		Color.a = 1.0f;
 	}
 
@@ -80,15 +80,15 @@ void MainCS(
 	float4 Result = Foreground;
 
 	BRANCH
-	if(BackingChannel == 0)
+	if(ScreenType == 0)
 	{
 		Result = RedScreenKeyer(Foreground.rgb, Background, ChannelWeights[1], ChannelWeights[2]);
 	}
-	else if(BackingChannel == 1)
+	else if(ScreenType == 1)
 	{
 		Result = GreenScreenKeyer(Foreground.rgb, Background, ChannelWeights[0], ChannelWeights[2]);
 	}
-	else if(BackingChannel == 2)
+	else if(ScreenType == 2)
 	{
 		Result = BlueScreenKeyer(Foreground.rgb, Background, ChannelWeights[0], ChannelWeights[1]);
 	}
diff --git a/Engine/Plugins/Experimental/Compositing/Composite/Source/Composite/Private/Passes/CompositePassDeltaKeyer.cpp b/Engine/Plugins/Experimental/Compositing/Composite/Source/Composite/Private/Passes/CompositePassDeltaKeyer.cpp
--- a/Engine/Plugins/Experimental/Compositing/Composite/Source/Composite/Private/Passes/CompositePassDeltaKeyer.cpp
+++ b/Engine/Plugins/Experimental/Compositing/Composite/Source/Composite/Private/Passes/CompositePassDeltaKeyer.cpp
@@ -33,7 +33,7 @@ class FCompositePassDeltaKeyerShader : public FGlobalShader
 		SHADER_PARAMETER(FVector3f, KeyColor)
 		SHADER_PARAMETER(FVector4f, Params0)
 		SHADER_PARAMETER(FVector4f, Params1)
-		SHADER_PARAMETER(uint32, BackingChannel)
+		SHADER_PARAMETER(uint32, ScreenType)
 		SHADER_PARAMETER(uint32, Visualization)
 		SHADER_PARAMETER(uint32, bInvertAlpha)
 	END_SHADER_PARAMETER_STRUCT()
@@ -91,7 +91,7 @@ namespace UE
 				PassParameters->KeyColor = FVector3f(KeyColor);
 				PassParameters->Params0 = Params0;
 				PassParameters->Params1 = Params1;
-				PassParameters->BackingChannel = BackingChannel;
+				PassParameters->ScreenType = static_cast<uint32>(ScreenType);
 				PassParameters->Visualization = static_cast<uint32>(Visualization);
 				PassParameters->bInvertAlpha = static_cast<uint32>(bInvertAlpha);
 	
@@ -125,7 +125,7 @@ namespace UE
 				return UE::CompositeCore::FPassTexture{ MoveTemp(Output), Metadata };
 			}
 
-			uint32 BackingChannel;
+			ECompositeDeltaKeyerScreenType ScreenType;
 			FLinearColor KeyColor;
 			FVector4f Params0;
 			FVector4f Params1;
@@ -138,11 +138,15 @@ namespace UE
 
 UCompositePassDeltaKeyer::UCompositePassDeltaKeyer(const FObjectInitializer& ObjectInitializer)
 	: Super(ObjectInitializer)
+	, ScreenType(ECompositeDeltaKeyerScreenType::Green)
 	, KeyColor(FLinearColor::Green)
-	, ChannelWeights(0.5f, 0.5f, 0.5f)
+	, CleanPlate(nullptr)
+	, RedWeight(0.5f)
+	, GreenWeight(0.5f)
+	, BlueWeight(0.5f)
 	, AlphaThreshold(0.0f, 1.0f)
 	, DespillStrength(1.0f)
-	, Visualization(ECompositeDeltaKeyerVisualization::None)
+	, Visualization(ECompositeDeltaKeyerVisualization::Key)
 	, bInvertAlpha(false)
 { 
 }
@@ -155,9 +159,9 @@ bool UCompositePassDeltaKeyer::GetProxy(const UE::CompositeCore::FPassInputDecl&
 
 	FCompositePassDeltaKeyerProxy* Proxy = InFrameAllocator.Create<FCompositePassDeltaKeyerProxy>(UE::CompositeCore::FPassInputDeclArray{ InputDecl });
 	Proxy->KeyColor = KeyColor;
-	Proxy->Params0 = FVector4f(ChannelWeights.X, ChannelWeights.Y, ChannelWeights.Z, DespillStrength);
+	Proxy->Params0 = FVector4f(RedWeight, GreenWeight, BlueWeight, DespillStrength);
 	Proxy->Params1 = FVector4f(AlphaThreshold.X, AlphaThreshold.Y, 0.0f, 0.0f);
-	Proxy->BackingChannel = static_cast<uint32>(FMath::Max3Index(KeyColor.R, KeyColor.G, KeyColor.B));
+	Proxy->ScreenType = ScreenType;
 	Proxy->bInvertAlpha = bInvertAlpha;
 	Proxy->Visualization = Visualization;
 	Proxy->CleanPlateWeakPtr = CleanPlate;
@@ -180,6 +184,7 @@ void UCompositePassDeltaKeyer::PreEditChange(FProperty* PropertyThatWillChange)
 	const bool bIsBooleanProperty = CastField<const FBoolProperty>(PropertyThatWillChange) != nullptr;
 
 	if (PropertyName != GET_MEMBER_NAME_CHECKED(ThisClass, bIsEnabled)
+		&& PropertyName != GET_MEMBER_NAME_CHECKED(ThisClass, ScreenType)
 		&& PropertyName != GET_MEMBER_NAME_CHECKED(ThisClass, KeyColor)
 		&& PropertyName != GET_MEMBER_NAME_CHECKED(ThisClass, CleanPlate)
 		&& !bIsBooleanProperty
diff --git a/Engine/Plugins/Experimental/Compositing/Composite/Source/Composite/Public/Passes/CompositePassDeltaKeyer.h b/Engine/Plugins/Experimental/Compositing/Composite/Source/Composite/Public/Passes/CompositePassDeltaKeyer.h
--- a/Engine/Plugins/Experimental/Compositing/Composite/Source/Composite/Public/Passes/CompositePassDeltaKeyer.h
+++ b/Engine/Plugins/Experimental/Compositing/Composite/Source/Composite/Public/Passes/CompositePassDeltaKeyer.h
@@ -8,12 +8,22 @@
 
 #define UE_API COMPOSITE_API
 
+/** Visualization mode, for inspecting your key.*/
 UENUM(BlueprintType)
 enum class ECompositeDeltaKeyerVisualization : uint8
 {
-	None = 0,
-	Alpha = 1,
-	Fill = 2,
+	Key = 0,
+	Fill = 1,
+	Alpha = 2
+};
+
+/** Background screen color type (red, green or blue).*/
+UENUM(BlueprintType)
+enum class ECompositeDeltaKeyerScreenType : uint8
+{
+	Red = 0,
+	Green = 1,
+	Blue = 2
 };
 
 /**
@@ -39,20 +49,32 @@ public:
 #endif
 
 public:
-	/** Background key color. Always set as we currently determine R/G/B backing type from its maximum channel. */
-	UPROPERTY(EditAnywhere, BlueprintReadWrite, Category = "", meta = (OnlyUpdateOnInteractionEnd))
+	/** Type of screen color (required). The keyer works best against red, green or blue backgrounds. */
+	UPROPERTY(EditAnywhere, BlueprintReadWrite, Category = "")
+	ECompositeDeltaKeyerScreenType ScreenType;
+
+	/** Static background key color.*/
+	UPROPERTY(EditAnywhere, BlueprintReadWrite, Category = "", meta = (OnlyUpdateOnInteractionEnd, EditCondition = "CleanPlate == nullptr"))
 	FLinearColor KeyColor;
 
 	/**
-	 * Clean plate background for calculating color differences, refining the static key color.
+	 * Clean plate background for calculating color differences per pixel, instead of the static key color.
 	 * Resolution must match the composite actor render resolution.
 	 */
 	UPROPERTY(EditAnywhere, BlueprintReadWrite, Category = "")
 	TObjectPtr<UTexture> CleanPlate;
 
-	/** Weights of the non-backing channels affecting the keying matte hardness. */
-	UPROPERTY(EditAnywhere, BlueprintReadWrite, Category = "", meta = (UIMin = 0.0f, UIMax = 1.5f))
-	FVector3f ChannelWeights;
+	/** Weight of the foreground red channel contributing to the key matte hardness. */
+	UPROPERTY(EditAnywhere, BlueprintReadWrite, Category = "", meta = (UIMin = 0.0f, UIMax = 1.0f, EditCondition = "ScreenType != ECompositeDeltaKeyerScreenType::Red"))
+	float RedWeight;
+
+	/** Weight of the foreground green channel contributing to the key matte hardness. */
+	UPROPERTY(EditAnywhere, BlueprintReadWrite, Category = "", meta = (UIMin = 0.0f, UIMax = 1.0f, EditCondition = "ScreenType != ECompositeDeltaKeyerScreenType::Green"))
+	float GreenWeight;
+
+	/** Weight of the foreground blue channel contributing to the key matte hardness. */
+	UPROPERTY(EditAnywhere, BlueprintReadWrite, Category = "", meta = (UIMin = 0.0f, UIMax = 1.0f, EditCondition = "ScreenType != ECompositeDeltaKeyerScreenType::Blue"))
+	float BlueWeight;
 
 	/** Thresholds any alpha value outside the specified range to zero or one respectively, with linear interpolation in-between. */
 	UPROPERTY(EditAnywhere, BlueprintReadWrite, Category = "", meta = (UIMin = 0.0f, UIMax = 1.0f))
