Composite: Added devignette parameter to delta keyer to compensate for corner darkening.

diff --git a/Engine/Plugins/Experimental/Compositing/Composite/Shaders/Private/CompositeDeltaKeyer.usf b/Engine/Plugins/Experimental/Compositing/Composite/Shaders/Private/CompositeDeltaKeyer.usf
--- a/Engine/Plugins/Experimental/Compositing/Composite/Shaders/Private/CompositeDeltaKeyer.usf
+++ b/Engine/Plugins/Experimental/Compositing/Composite/Shaders/Private/CompositeDeltaKeyer.usf
@@ -2,7 +2,6 @@
 
 #include "/Engine/Private/Common.ush"
 #include "/Engine/Private/ScreenPass.ush"
-#include "/Engine/Private/PostProcessCommon.ush"
 
 #include "CompositeDeltaKeyerCommon.ush"
 
@@ -28,6 +27,11 @@ uint ScreenType;
 uint Visualization;
 uint bInvertAlpha;
 
+float2 UVToScreenPos(float2 UV, float2 Extent)
+{
+	return (UV * Extent - Input_ScreenPosToViewportBias) / Input_ScreenPosToViewportScale;
+}
+
 // Output color function with debug convenience
 float4 Output(float4 Color)
 {
@@ -69,6 +73,7 @@ void MainCS(
 	const float DespillStrength = Params0.w;
 	const float AlphaMin = Params1.x;
 	const float AlphaMax = Params1.y;
+	const float DevignetteStrength = Params1.z;
 	
 	float4 Foreground = InputTexture.SampleLevel(InputSampler, UV, 0);
 #if USE_CLEAN_PLATE
@@ -79,18 +84,28 @@ void MainCS(
 
 	float4 Result = Foreground;
 
+	BRANCH
+	if(DevignetteStrength > ZeroEpsilon)
+	{
+		float2 ScreenPos = UVToScreenPos(UV, Input_Extent);
+		float AspectRatio = Output_ViewportSize.y * Output_ViewportSizeInverse.x;
+		float VignetteMask = Vignette(ScreenPos, AspectRatio, DevignetteStrength);
+
+		Result.rgb *= rcp(VignetteMask);
+	}
+
 	BRANCH
 	if(ScreenType == 0)
 	{
-		Result = RedScreenKeyer(Foreground.rgb, Background, ChannelWeights[1], ChannelWeights[2]);
+		Result = RedScreenKeyer(Result.rgb, Background, ChannelWeights[1], ChannelWeights[2]);
 	}
 	else if(ScreenType == 1)
 	{
-		Result = GreenScreenKeyer(Foreground.rgb, Background, ChannelWeights[0], ChannelWeights[2]);
+		Result = GreenScreenKeyer(Result.rgb, Background, ChannelWeights[0], ChannelWeights[2]);
 	}
 	else if(ScreenType == 2)
 	{
-		Result = BlueScreenKeyer(Foreground.rgb, Background, ChannelWeights[0], ChannelWeights[1]);
+		Result = BlueScreenKeyer(Result.rgb, Background, ChannelWeights[0], ChannelWeights[1]);
 	}
 
 	BRANCH
diff --git a/Engine/Plugins/Experimental/Compositing/Composite/Shaders/Private/CompositeDeltaKeyerCommon.ush b/Engine/Plugins/Experimental/Compositing/Composite/Shaders/Private/CompositeDeltaKeyerCommon.ush
--- a/Engine/Plugins/Experimental/Compositing/Composite/Shaders/Private/CompositeDeltaKeyerCommon.ush
+++ b/Engine/Plugins/Experimental/Compositing/Composite/Shaders/Private/CompositeDeltaKeyerCommon.ush
@@ -5,6 +5,7 @@
 #include "/Engine/Private/ACES/ACESCommon.ush"
 #include "/Engine/Private/ColorSpace.ush"
 #include "/Engine/Private/GammaCorrectionCommon.ush"
+#include "/Engine/Private/PostProcessCommon.ush"
 
 // CIE 1931 D65 white point
 static const float3 XYZ_n = xyY_2_XYZ(float3(0.3127, 0.3290, 1.0));
@@ -161,6 +162,13 @@ float3 BlueDespillAvg(float3 Color)
 	);
 }
 
+float Vignette(float2 ScreenPos, float AspectRatio, float Intensity)
+{
+	// Scale vignette to always be a circle with consistent corner intensity.
+	float2 Vignette = VignetteSpace(ScreenPos, AspectRatio);
+	return ComputeVignetteMask(Vignette, Intensity);
+}
+
 float4 GreenScreenKeyer(
 	float3 Foreground,
 	float3 Background,
diff --git a/Engine/Plugins/Experimental/Compositing/Composite/Source/Composite/Private/Passes/CompositePassDeltaKeyer.cpp b/Engine/Plugins/Experimental/Compositing/Composite/Source/Composite/Private/Passes/CompositePassDeltaKeyer.cpp
--- a/Engine/Plugins/Experimental/Compositing/Composite/Source/Composite/Private/Passes/CompositePassDeltaKeyer.cpp
+++ b/Engine/Plugins/Experimental/Compositing/Composite/Source/Composite/Private/Passes/CompositePassDeltaKeyer.cpp
@@ -160,7 +160,7 @@ bool UCompositePassDeltaKeyer::GetProxy(const UE::CompositeCore::FPassInputDecl&
 	FCompositePassDeltaKeyerProxy* Proxy = InFrameAllocator.Create<FCompositePassDeltaKeyerProxy>(UE::CompositeCore::FPassInputDeclArray{ InputDecl });
 	Proxy->KeyColor = KeyColor;
 	Proxy->Params0 = FVector4f(RedWeight, GreenWeight, BlueWeight, DespillStrength);
-	Proxy->Params1 = FVector4f(AlphaThreshold.X, AlphaThreshold.Y, 0.0f, 0.0f);
+	Proxy->Params1 = FVector4f(AlphaThreshold.X, AlphaThreshold.Y, DevignetteStrength, 0.0f);
 	Proxy->ScreenType = ScreenType;
 	Proxy->bInvertAlpha = bInvertAlpha;
 	Proxy->Visualization = Visualization;
diff --git a/Engine/Plugins/Experimental/Compositing/Composite/Source/Composite/Public/Passes/CompositePassDeltaKeyer.h b/Engine/Plugins/Experimental/Compositing/Composite/Source/Composite/Public/Passes/CompositePassDeltaKeyer.h
--- a/Engine/Plugins/Experimental/Compositing/Composite/Source/Composite/Public/Passes/CompositePassDeltaKeyer.h
+++ b/Engine/Plugins/Experimental/Compositing/Composite/Source/Composite/Public/Passes/CompositePassDeltaKeyer.h
@@ -84,6 +84,10 @@ public:
 	UPROPERTY(EditAnywhere, BlueprintReadWrite, Category = "", meta = (UIMin = 0.0f, UIMax = 1.0f))
 	float DespillStrength;
 
+	/** Weight of the foreground blue channel contributing to the key matte hardness. */
+	UPROPERTY(EditAnywhere, BlueprintReadWrite, Category = "", meta = (UIMin = 0.0f, UIMax = 1.0f))
+	float DevignetteStrength;
+
 	/** Vizualize the alpha key or fill. */
 	UPROPERTY(EditAnywhere, BlueprintReadWrite, Category = "")
 	ECompositeDeltaKeyerVisualization Visualization;
