Composure: Added Composure panel that displays any composite actors in the level and their composite layers in a tree view.

diff --git a/Engine/Plugins/Experimental/Compositing/Composite/Source/Composite/Private/Layers/CompositeLayerPlate.cpp b/Engine/Plugins/Experimental/Compositing/Composite/Source/Composite/Private/Layers/CompositeLayerPlate.cpp
--- a/Engine/Plugins/Experimental/Compositing/Composite/Source/Composite/Private/Layers/CompositeLayerPlate.cpp
+++ b/Engine/Plugins/Experimental/Compositing/Composite/Source/Composite/Private/Layers/CompositeLayerPlate.cpp
@@ -139,7 +139,7 @@ void UCompositeLayerPlate::PostInitProperties()
 {
 	Super::PostInitProperties();
 
-	if (!HasAnyFlags(RF_ClassDefaultObject))
+	if (!HasAnyFlags(RF_ClassDefaultObject | RF_Transient))
 	{
 		ACompositeActor* CompositeActor = GetTypedOuter<ACompositeActor>();
 		
diff --git a/Engine/Plugins/Experimental/Compositing/Composite/Source/Composite/Public/Layers/CompositeLayerBase.h b/Engine/Plugins/Experimental/Compositing/Composite/Source/Composite/Public/Layers/CompositeLayerBase.h
--- a/Engine/Plugins/Experimental/Compositing/Composite/Source/Composite/Public/Layers/CompositeLayerBase.h
+++ b/Engine/Plugins/Experimental/Compositing/Composite/Source/Composite/Public/Layers/CompositeLayerBase.h
@@ -57,21 +57,21 @@ public:
 
 public:
 	/** Whether or not the pass is enabled. */
-	UPROPERTY(EditAnywhere, BlueprintReadWrite, Category = "", meta = (DisplayPriority = "1"))
+	UPROPERTY(EditAnywhere, BlueprintReadWrite, Category = "", meta = (DisplayPriority = "2"))
 	bool bIsEnabled = true;
 
 	/** Whether or not the pass is solo. */
-	UPROPERTY(EditAnywhere, BlueprintReadWrite, Category = "", meta = (DisplayPriority = "2"))
+	UPROPERTY()
 	bool bIsSolo = false;
 
 #if WITH_EDITORONLY_DATA
 	/** Pass layer name. */
-	UPROPERTY()
+	UPROPERTY(EditAnywhere, BlueprintReadWrite, Category = "", meta = (DisplayPriority = "1"))
 	FString Name;
 #endif
 
 	/** Merge operation applied on Input0 with Input1. */
-	UPROPERTY(EditAnywhere, BlueprintReadWrite, Category = "", meta = (DisplayPriority = "4"))
+	UPROPERTY(EditAnywhere, BlueprintReadWrite, Category = "", meta = (DisplayPriority = "3"))
 	ECompositeCoreMergeOp Operation = ECompositeCoreMergeOp::Over;
 
 protected:
diff --git a/Engine/Plugins/Experimental/Compositing/Composite/Source/Composite/Public/Layers/CompositeLayerSceneCapture.h b/Engine/Plugins/Experimental/Compositing/Composite/Source/Composite/Public/Layers/CompositeLayerSceneCapture.h
--- a/Engine/Plugins/Experimental/Compositing/Composite/Source/Composite/Public/Layers/CompositeLayerSceneCapture.h
+++ b/Engine/Plugins/Experimental/Compositing/Composite/Source/Composite/Public/Layers/CompositeLayerSceneCapture.h
@@ -63,7 +63,7 @@ private:
 	 * 
 	 * Automatically sets the scene capture's ShowOnlyActors list when updated.
 	*/
-	UPROPERTY(EditAnywhere, BlueprintGetter = GetActors, BlueprintSetter = SetActors, Category = "Composite", meta = (AllowPrivateAccess = true))
+	UPROPERTY(EditAnywhere, BlueprintGetter = GetActors, BlueprintSetter = SetActors, Category = "", meta = (AllowPrivateAccess = true))
 	TArray<TObjectPtr<AActor>> Actors;
 
 	/**
@@ -71,7 +71,7 @@ private:
 	 * 
 	 * Automatically sets the scene capture's CaptureSource, bRenderInMainRenderer & bIgnoreScreenPercentage when updated.
 	 */
-	UPROPERTY(EditAnywhere, BlueprintGetter = IsCustomRenderPass, BlueprintSetter = SetCustomRenderPass, Category = "Composite", meta = (AllowPrivateAccess = true))
+	UPROPERTY(EditAnywhere, BlueprintGetter = IsCustomRenderPass, BlueprintSetter = SetCustomRenderPass, Category = "", meta = (AllowPrivateAccess = true))
 	bool bCustomRenderPass;
 };
 
diff --git a/Engine/Plugins/Experimental/Compositing/Composite/Source/Composite/Public/Layers/CompositeLayerShadowReflection.h b/Engine/Plugins/Experimental/Compositing/Composite/Source/Composite/Public/Layers/CompositeLayerShadowReflection.h
--- a/Engine/Plugins/Experimental/Compositing/Composite/Source/Composite/Public/Layers/CompositeLayerShadowReflection.h
+++ b/Engine/Plugins/Experimental/Compositing/Composite/Source/Composite/Public/Layers/CompositeLayerShadowReflection.h
@@ -59,7 +59,7 @@ public:
 	 * * bAffectIndirectLightingWhileHidden
 	 * * bCastHiddenShadow
 	*/
-	UPROPERTY(EditAnywhere, BlueprintReadWrite, Category = "Composite", AdvancedDisplay, meta = (DisplayAfter = "Actors", DisplayName = "Auto-Configure Actors"))
+	UPROPERTY(EditAnywhere, BlueprintReadWrite, Category = "", meta = (DisplayAfter = "Actors", DisplayName = "Auto-Configure Actors"))
 	ECompositeHiddenInSceneCaptureConfiguration AutoConfigureActors;
 
 private:
@@ -67,7 +67,7 @@ private:
 	void ConfigurePrimitivesForHiddenShadowReflectionCatching();
 
 	/** List of reflection & shadow caster actors. */
-	UPROPERTY(EditAnywhere, BlueprintGetter = GetActors, BlueprintSetter = SetActors, Category = "Composite")
+	UPROPERTY(EditAnywhere, BlueprintGetter = GetActors, BlueprintSetter = SetActors, Category = "")
 	TArray<TObjectPtr<AActor>> Actors;
 };
 
diff --git a/Engine/Plugins/Experimental/Compositing/Composite/Source/CompositeEditor/CompositeEditor.Build.cs b/Engine/Plugins/Experimental/Compositing/Composite/Source/CompositeEditor/CompositeEditor.Build.cs
--- a/Engine/Plugins/Experimental/Compositing/Composite/Source/CompositeEditor/CompositeEditor.Build.cs
+++ b/Engine/Plugins/Experimental/Compositing/Composite/Source/CompositeEditor/CompositeEditor.Build.cs
@@ -19,10 +19,18 @@ public class CompositeEditor : ModuleRules
 		PrivateDependencyModuleNames.AddRange(
 			new string[]
 			{
+				"ApplicationCore",
+				"CompositeCore",
 				"CoreUObject",
 				"Engine",
+				"InputCore",
+				"LevelEditor",
 				"UnrealEd",
                 "Slate",
+                "SlateCore",
+                "ToolMenus",
+                "ToolWidgets",
+                "WorkspaceMenuStructure",
 			}
 		);
 	}
diff --git a/Engine/Plugins/Experimental/Compositing/Composite/Source/CompositeEditor/Private/CompositeEditorModule.cpp b/Engine/Plugins/Experimental/Compositing/Composite/Source/CompositeEditor/Private/CompositeEditorModule.cpp
--- a/Engine/Plugins/Experimental/Compositing/Composite/Source/CompositeEditor/Private/CompositeEditorModule.cpp
+++ b/Engine/Plugins/Experimental/Compositing/Composite/Source/CompositeEditor/Private/CompositeEditorModule.cpp
@@ -11,6 +11,7 @@
 #include "Widgets/Notifications/SNotificationList.h"
 
 #include "CompositeActor.h"
+#include "UI/SCompositeEditorPanel.h"
 
 #define LOCTEXT_NAMESPACE "CompositeEditorModule"
 
@@ -19,10 +20,14 @@ DEFINE_LOG_CATEGORY(LogCompositeEditor);
 void FCompositeEditorModule::StartupModule()
 {
 	FCoreDelegates::OnPostEngineInit.AddRaw(this, &FCompositeEditorModule::OnPostEngineInit);
+
+	RegisterTabSpawners();
 }
 
 void FCompositeEditorModule::ShutdownModule()
 {
+	UnregisterTabSpawners();
+	
 	if (GEditor != nullptr)
 	{
 		GEditor->OnLevelActorAdded().RemoveAll(this);
@@ -74,6 +79,24 @@ void FCompositeEditorModule::OnLevelActorAdded(AActor* InActor)
 	}
 }
 
+void FCompositeEditorModule::RegisterTabSpawners()
+{
+	FLevelEditorModule& LevelEditorModule = FModuleManager::GetModuleChecked<FLevelEditorModule>("LevelEditor");
+
+	OnTabManagerChangedDelegateHandle = LevelEditorModule.OnTabManagerChanged().AddLambda([this]()
+	{
+		SCompositeEditorPanel::RegisterTabSpawner();
+	});
+}
+
+void FCompositeEditorModule::UnregisterTabSpawners()
+{
+	SCompositeEditorPanel::UnregisterTabSpawner();
+
+	FLevelEditorModule& LevelEditorModule = FModuleManager::GetModuleChecked<FLevelEditorModule>("LevelEditor");
+	LevelEditorModule.OnTabManagerChanged().Remove(OnTabManagerChangedDelegateHandle);
+}
+
 IMPLEMENT_MODULE(FCompositeEditorModule, CompositeEditor)
 
 #undef LOCTEXT_NAMESPACE
diff --git a/Engine/Plugins/Experimental/Compositing/Composite/Source/CompositeEditor/Private/UI/SCompositeEditorPanel.cpp b/Engine/Plugins/Experimental/Compositing/Composite/Source/CompositeEditor/Private/UI/SCompositeEditorPanel.cpp
new file mode 100644
--- /dev/null
+++ b/Engine/Plugins/Experimental/Compositing/Composite/Source/CompositeEditor/Private/UI/SCompositeEditorPanel.cpp
@@ -0,0 +1,451 @@
+﻿// Copyright Epic Games, Inc. All Rights Reserved.
+
+#include "SCompositeEditorPanel.h"
+
+#include "CompositeActor.h"
+#include "DetailCategoryBuilder.h"
+#include "DetailLayoutBuilder.h"
+#include "DetailWidgetRow.h"
+#include "IDetailChildrenBuilder.h"
+#include "IDetailCustomization.h"
+#include "LevelEditor.h"
+#include "SCompositePanelLayerTree.h"
+#include "ScopedTransaction.h"
+#include "WorkspaceMenuStructure.h"
+#include "WorkspaceMenuStructureModule.h"
+#include "Widgets/Docking/SDockTab.h"
+#include "Widgets/Images/SImage.h"
+#include "Widgets/Input/SComboBox.h"
+#include "Widgets/Input/SEditableTextBox.h"
+
+#define LOCTEXT_NAMESPACE "SCompositeEditorPanel"
+
+namespace CompositingEditorPanel
+{
+	static const FName CompositeEditorTabName = "CompositeEditorTab";
+
+	TSharedRef<SDockTab> CreateTab(const FSpawnTabArgs& InArgs)
+	{
+		return SNew(SDockTab)
+			.Label(LOCTEXT("PanelTabLabel", "Composure"))
+			.TabRole(ETabRole::PanelTab)
+			[
+				SNew(SCompositeEditorPanel)
+			];
+	}
+}
+
+class FResolutionTypeCustomization : public IPropertyTypeCustomization
+{
+private:
+	struct FNamedResolution
+	{
+		FText Name = FText::GetEmpty();
+		FIntPoint Resolution = FIntPoint::ZeroValue;
+		FText ToolTip = FText::GetEmpty();
+			
+		FNamedResolution(const FText& InName, const FIntPoint& InResolution, const FText& InToolTip)
+			: Name(InName)
+			, Resolution(InResolution)
+			, ToolTip(InToolTip)
+		{ }
+	};
+	using FNamedResolutionPtr = TSharedPtr<FNamedResolution>;
+	
+	static TArray<FNamedResolutionPtr> NamedResolutions;
+	
+public:
+	virtual void CustomizeHeader(TSharedRef<IPropertyHandle> InPropertyHandle, FDetailWidgetRow& InHeaderRow, IPropertyTypeCustomizationUtils& InCustomizationUtils) override
+	{
+		PropertyHandle = InPropertyHandle;
+
+		InHeaderRow
+		.NameContent()
+		[
+			PropertyHandle->CreatePropertyNameWidget()
+		]
+		.ValueContent()
+		[
+			SAssignNew(ResolutionComboBox, SComboBox<FNamedResolutionPtr>)
+			.OptionsSource(&NamedResolutions)
+			.OnGenerateWidget_Lambda([](FNamedResolutionPtr InItem)
+			{
+				return SNew(STextBlock)
+					.Font(IDetailLayoutBuilder::GetDetailFont())
+					.Text(InItem->Name)
+					.ToolTipText(InItem->ToolTip);
+			})
+			.OnSelectionChanged(this, &FResolutionTypeCustomization::SetResolution)
+			.InitiallySelectedItem(NamedResolutions[1])
+			[
+				SNew(STextBlock)
+				.Font(IDetailLayoutBuilder::GetDetailFont())
+				.Text(this, &FResolutionTypeCustomization::GetResolutionText)
+			]
+		];
+
+		InHeaderRow.OverrideResetToDefault(FResetToDefaultOverride::Create(
+			FIsResetToDefaultVisible::CreateLambda([this](TSharedPtr<IPropertyHandle>)
+			{
+				return ResolutionComboBox->GetSelectedItem() != NamedResolutions[1];
+			}), FResetToDefaultHandler::CreateLambda([this](TSharedPtr<IPropertyHandle>)
+			{
+				ResolutionComboBox->SetSelectedItem(NamedResolutions[1]);
+			})));
+	}
+
+	virtual void CustomizeChildren(TSharedRef<IPropertyHandle> InPropertyHandle, IDetailChildrenBuilder& InChildBuilder, IPropertyTypeCustomizationUtils& InCustomizationUtils) override
+	{
+		XPropertyHandle = InPropertyHandle->GetChildHandle(GET_MEMBER_NAME_CHECKED(FIntPoint, X));
+		YPropertyHandle = InPropertyHandle->GetChildHandle(GET_MEMBER_NAME_CHECKED(FIntPoint, Y));
+
+		InChildBuilder.AddCustomRow(InPropertyHandle->GetPropertyDisplayName())
+		.NameContent()
+		[
+			XPropertyHandle->CreatePropertyNameWidget()
+		]
+		.ValueContent()
+		[
+			XPropertyHandle->CreatePropertyValueWidget()
+		]
+		.ExtensionContent()
+		[
+			SNew(SButton)
+			.OnClicked_Lambda([this]()
+			{
+				if (LockedAspectRatio.IsSet())
+				{
+					LockedAspectRatio.Reset();
+				}
+				else
+				{
+					double AspectRatio;
+					if (GetAspectRatio(AspectRatio))
+					{
+						LockedAspectRatio = AspectRatio;
+					}
+				}
+				
+				return FReply::Handled();
+			})
+			.ContentPadding(FMargin(0, 0, 4, 0))
+			.ButtonStyle(FAppStyle::Get(), "NoBorder")
+			[
+				SNew(SImage)
+				.ColorAndOpacity(FSlateColor::UseForeground())
+				.Image_Lambda([this]
+				{
+					return LockedAspectRatio.IsSet() ? FAppStyle::GetBrush(TEXT("Icons.Link")) : FAppStyle::GetBrush(TEXT("Icons.Unlink"));
+				})
+				.ToolTipText_Lambda([this]
+				{
+					return LockedAspectRatio.IsSet() ?
+					FText::Format(LOCTEXT("LockedAspectRatioTooltipFormat", "Click to unlock Aspect Ratio ({0})"), FText::AsNumber(LockedAspectRatio.GetValue())) :
+					LOCTEXT("UnlockedAspectRatioTooltip", "Click to lock the Aspect Ratio");
+				})
+			]
+		];
+		
+		InChildBuilder.AddCustomRow(InPropertyHandle->GetPropertyDisplayName())
+		.NameContent()
+		[
+			YPropertyHandle->CreatePropertyNameWidget()
+		]
+		.ValueContent()
+		[
+			YPropertyHandle->CreatePropertyValueWidget()
+		];
+
+		XPropertyHandle->SetOnPropertyValueChanged(FSimpleDelegate::CreateLambda([this]
+		{
+			if (!bSettingNamedResolution)
+			{
+				ResolutionComboBox->SetSelectedItem(NamedResolutions.Last());
+			}
+
+			if (LockedAspectRatio.IsSet())
+			{
+				SetDimensionFromAspectRatio(XPropertyHandle, YPropertyHandle, 1.0 / LockedAspectRatio.GetValue());
+			}
+		}));
+		
+		YPropertyHandle->SetOnPropertyValueChanged(FSimpleDelegate::CreateLambda([this]
+		{
+			if (!bSettingNamedResolution)
+			{
+				ResolutionComboBox->SetSelectedItem(NamedResolutions.Last());
+			}
+
+			if (LockedAspectRatio.IsSet())
+			{
+				SetDimensionFromAspectRatio(YPropertyHandle, XPropertyHandle, LockedAspectRatio.GetValue());
+			}
+		}));
+	}
+	
+private:
+	FText GetResolutionText() const
+	{
+		FIntPoint Resolution;
+		FPropertyAccess::Result XResult = PropertyHandle->GetChildHandle(GET_MEMBER_NAME_CHECKED(FIntPoint, X))->GetValue(Resolution.X);
+		FPropertyAccess::Result YResult = PropertyHandle->GetChildHandle(GET_MEMBER_NAME_CHECKED(FIntPoint, Y))->GetValue(Resolution.Y);
+
+		if (XResult == FPropertyAccess::Fail || YResult == FPropertyAccess::Fail)
+		{
+			return FText::GetEmpty();
+		}
+
+		if (XResult == FPropertyAccess::MultipleValues || YResult == FPropertyAccess::MultipleValues)
+		{
+			return LOCTEXT("MultipleValues", "Multiple Values");
+		}
+
+		FNamedResolutionPtr SelectedResolution = ResolutionComboBox->GetSelectedItem();
+		if (SelectedResolution->Resolution != FIntPoint::ZeroValue)
+		{
+			return SelectedResolution->Name;
+		}
+
+		FNumberFormattingOptions FormattingOptions;
+		FormattingOptions.UseGrouping = false;
+		return FText::Format(LOCTEXT("CustomResolutionTextFormat", "Custom - {0}x{1}"),
+			FText::AsNumber(Resolution.X, &FormattingOptions),
+			FText::AsNumber(Resolution.Y, &FormattingOptions));
+	}
+
+	void SetResolution(FNamedResolutionPtr InNamedResolution, ESelectInfo::Type SelectInfo)
+	{
+		if (InNamedResolution->Resolution != FIntPoint::ZeroValue)
+		{
+			TGuardValue<bool> SettingNamedResolution(bSettingNamedResolution, true);
+			PropertyHandle->GetChildHandle(GET_MEMBER_NAME_CHECKED(FIntPoint, X))->SetValue(InNamedResolution->Resolution.X);
+			PropertyHandle->GetChildHandle(GET_MEMBER_NAME_CHECKED(FIntPoint, Y))->SetValue(InNamedResolution->Resolution.Y);
+		}
+	}
+
+	bool GetAspectRatio(double& OutAspectRatio) const
+	{
+		int32 X = 0;
+		int32 Y = 0;
+
+		if (XPropertyHandle->GetValue(X) != FPropertyAccess::Success)
+		{
+			return false;
+		}
+
+		if (YPropertyHandle->GetValue(Y) != FPropertyAccess::Success)
+		{
+			return false;
+		}
+
+		if (Y == 0)
+		{
+			return false;
+		}
+
+		OutAspectRatio = (double)X / (double)Y;
+		return true;
+	}
+
+	void SetDimensionFromAspectRatio(const TSharedPtr<IPropertyHandle>& SrcDimHandle, const TSharedPtr<IPropertyHandle>& DestDimHandle, double AspectRatio)
+	{
+		if (bAspectRatioRecursionGuard)
+		{
+			return;
+		}
+
+		TGuardValue<bool> RecursionGuard(bAspectRatioRecursionGuard, true);
+		
+		int32 SrcDim;
+		if (SrcDimHandle->GetValue(SrcDim) != FPropertyAccess::Success)
+		{
+			return;
+		}
+
+		const double ScaledSrcDim = SrcDim * AspectRatio;
+
+		// Round to the nearest even number
+		int32 DestDim = FMath::FloorToInt(ScaledSrcDim);
+		if (DestDim % 2 != 0)
+		{
+			DestDim = FMath::CeilToInt(ScaledSrcDim);
+		}
+
+		DestDimHandle->SetValue(DestDim);
+	}
+	
+private:
+	TSharedPtr<IPropertyHandle> PropertyHandle;
+	TSharedPtr<IPropertyHandle> XPropertyHandle;
+	TSharedPtr<IPropertyHandle> YPropertyHandle;
+
+	TSharedPtr<SComboBox<FNamedResolutionPtr>> ResolutionComboBox;
+	
+	TOptional<double> LockedAspectRatio = TOptional<double>();
+
+	bool bAspectRatioRecursionGuard = false;
+	bool bSettingNamedResolution = false;
+};
+
+TArray<TSharedPtr<FResolutionTypeCustomization::FNamedResolution>> FResolutionTypeCustomization::NamedResolutions =
+{
+	MakeShared<FNamedResolution>(LOCTEXT("720pResolutionName", "720p (HD) - 1280x720"), FIntPoint(1280, 720), LOCTEXT("Resolution720pToolTip", "High Definition")),
+	MakeShared<FNamedResolution>(LOCTEXT("1080pResolutionName", "1080p (FHD) - 1920x1080"), FIntPoint(1920, 1080), LOCTEXT("Resolution1080pToolTip", "Full high definition")),
+	MakeShared<FNamedResolution>(LOCTEXT("1440pResolutionName", "1440p (QHD) - 2560x1440"), FIntPoint(2560, 1440), LOCTEXT("Resolution1440pToolTip", "Quad high definition")),
+	MakeShared<FNamedResolution>(LOCTEXT("2160pResolutionName", "2160p (4K UHD) - 3840x2160"), FIntPoint(3840, 2160), LOCTEXT("Resolution2160pToolTip", "Ultra high definition (4k)")),
+	MakeShared<FNamedResolution>(LOCTEXT("CustomResolutionName", "Custom"), FIntPoint::ZeroValue, FText::GetEmpty())
+};
+
+class FCompositeActorPanelDetailCustomization : public IDetailCustomization
+{
+private:
+	class FResolutionPropertyIdentifier : public IPropertyTypeIdentifier
+	{
+		virtual bool IsPropertyTypeCustomized(const IPropertyHandle& PropertyHandle) const override
+		{
+			return PropertyHandle.GetProperty()->GetNameCPP() == GET_MEMBER_NAME_STRING_CHECKED(ACompositeActor, RenderResolution);
+		}
+	};
+	
+public:
+	virtual void CustomizeDetails(IDetailLayoutBuilder& DetailBuilder) override
+	{
+		DetailBuilder.RegisterInstancedCustomPropertyTypeLayout(NAME_IntPoint, FOnGetPropertyTypeCustomizationInstance::CreateLambda([]
+			{
+				return MakeShared<FResolutionTypeCustomization>();
+			}),
+			MakeShared<FResolutionPropertyIdentifier>());
+		
+		DetailBuilder.GetObjectsBeingCustomized(ObjectsBeingCustomized);
+		
+		// Hide all categories besides the Composite category
+		static const FName CompositeCategoryName = "Composite";
+		
+		TArray<FName> CategoryNames;
+		DetailBuilder.GetCategoryNames(CategoryNames);
+
+		for (const FName& CategoryName : CategoryNames)
+		{
+			if (CategoryName == CompositeCategoryName)
+			{
+				continue;
+			}
+			
+			DetailBuilder.HideCategory(CategoryName);
+		}
+
+		// TransformCommon is a custom category that doesn't get returned by GetCategoryNames that also needs to be hidden
+		DetailBuilder.HideCategory(TEXT("TransformCommon"));
+
+		IDetailCategoryBuilder& CompositeCategory = DetailBuilder.EditCategory(CompositeCategoryName, LOCTEXT("CompositeActorCategoryName", "Composite Actor"));
+
+		CompositeCategory.AddCustomRow(LOCTEXT("CompositeActorNameFilterText", "Name"))
+			.NameContent()
+			[
+				SNew(STextBlock)
+				.Text(LOCTEXT("CompositeActorNameProperty", "Name"))
+				.Font(FAppStyle::Get().GetFontStyle("PropertyWindow.NormalFont"))
+			]
+			.ValueContent()
+			[
+				SNew(SEditableTextBox)
+				.Text_Lambda([this]
+				{
+					if (ObjectsBeingCustomized.Num() == 1 && ObjectsBeingCustomized[0].IsValid() && ObjectsBeingCustomized[0]->IsA<AActor>())
+					{
+						return FText::FromString(Cast<AActor>(ObjectsBeingCustomized[0])->GetActorLabel());
+					}
+					if (ObjectsBeingCustomized.Num() > 1)
+					{
+						return LOCTEXT("MultipleValuesLabel", "Multiple Values");
+					}
+
+					return FText::GetEmpty();
+				})
+				.OnTextCommitted_Lambda([this](const FText& InText, ETextCommit::Type InCommitType)
+				{
+					if (ObjectsBeingCustomized.Num() == 1 && ObjectsBeingCustomized[0].IsValid() && ObjectsBeingCustomized[0]->IsA<AActor>())
+					{
+						FScopedTransaction RenameActorTransaction(LOCTEXT("RenameActorTransaction", "Rename Composite Actor"));
+						FActorLabelUtilities::RenameExistingActor(Cast<AActor>(ObjectsBeingCustomized[0]), InText.ToString(), true);
+					}
+				})
+				.IsEnabled_Lambda([this]
+				{
+					return ObjectsBeingCustomized.Num() == 1;
+				})
+			];
+		
+		DetailBuilder.HideProperty(GET_MEMBER_NAME_CHECKED(ACompositeActor, CompositeLayers));
+	}
+
+private:
+	TArray<TWeakObjectPtr<UObject>> ObjectsBeingCustomized;
+
+	TSharedPtr<IPropertyHandle> RenderResolutionHandle;
+};
+
+void SCompositeEditorPanel::RegisterTabSpawner()
+{
+	FLevelEditorModule& LevelEditorModule = FModuleManager::LoadModuleChecked<FLevelEditorModule>("LevelEditor");
+	if (TSharedPtr<FTabManager> LevelEditorTabManager = LevelEditorModule.GetLevelEditorTabManager())
+	{
+		LevelEditorTabManager->RegisterTabSpawner(CompositingEditorPanel::CompositeEditorTabName, FOnSpawnTab::CreateStatic(&CompositingEditorPanel::CreateTab))
+			.SetDisplayName(LOCTEXT("PanelDisplayName", "Composure"))
+			.SetTooltipText(LOCTEXT("PanelTooltipText", "Panel for viewing and editing compositing actors in the current level"))
+			.SetGroup(WorkspaceMenu::GetMenuStructure().GetLevelEditorVirtualProductionCategory());
+			//.SetIcon(FSlateIcon(FVPUtilitiesEditorStyle::GetStyleSetName(), "TabIcons.Genlock.Small")); // TODO: find the correct icon
+	}
+}
+
+void SCompositeEditorPanel::UnregisterTabSpawner()
+{
+	FLevelEditorModule& LevelEditorModule = FModuleManager::LoadModuleChecked<FLevelEditorModule>("LevelEditor");
+	if (TSharedPtr<FTabManager> LevelEditorTabManager = LevelEditorModule.GetLevelEditorTabManager())
+	{
+		LevelEditorTabManager->UnregisterTabSpawner(CompositingEditorPanel::CompositeEditorTabName);
+	}
+}
+
+void SCompositeEditorPanel::Construct(const FArguments& InArgs)
+{
+	FDetailsViewArgs DetailsViewArgs;
+	DetailsViewArgs.NameAreaSettings = FDetailsViewArgs::HideNameArea;
+	DetailsViewArgs.bAllowMultipleTopLevelObjects = true;
+	DetailsViewArgs.bShowPropertyMatrixButton = false;
+	
+	DetailsView = FModuleManager::GetModuleChecked<FPropertyEditorModule>("PropertyEditor").CreateDetailView(DetailsViewArgs);
+	DetailsView->RegisterInstancedCustomPropertyLayout(ACompositeActor::StaticClass(), FOnGetDetailCustomizationInstance::CreateLambda([]
+	{
+		return MakeShared<FCompositeActorPanelDetailCustomization>();
+	}));
+	
+	ChildSlot
+	[
+		SNew(SSplitter)
+		.Orientation(EOrientation::Orient_Vertical)
+
+		+SSplitter::Slot()
+		.Value(0.25f)
+		[
+			SNew(SCompositePanelLayerTree)
+			.OnSelectionChanged(this, &SCompositeEditorPanel::OnLayerSelectionChanged)
+		]
+
+		+SSplitter::Slot()
+		.Value(0.75f)
+		[
+			DetailsView.ToSharedRef()
+		]
+	];
+}
+
+void SCompositeEditorPanel::OnLayerSelectionChanged(const TArray<UObject*>& SelectedLayers)
+{
+	if (DetailsView.IsValid())
+	{
+		DetailsView->SetObjects(SelectedLayers);
+	}
+}
+
+#undef LOCTEXT_NAMESPACE
diff --git a/Engine/Plugins/Experimental/Compositing/Composite/Source/CompositeEditor/Private/UI/SCompositeEditorPanel.h b/Engine/Plugins/Experimental/Compositing/Composite/Source/CompositeEditor/Private/UI/SCompositeEditorPanel.h
new file mode 100644
--- /dev/null
+++ b/Engine/Plugins/Experimental/Compositing/Composite/Source/CompositeEditor/Private/UI/SCompositeEditorPanel.h
@@ -0,0 +1,38 @@
+﻿// Copyright Epic Games, Inc. All Rights Reserved.
+
+#pragma once
+
+#include "Widgets/SCompoundWidget.h"
+
+class IDetailsView;
+class SCompositePanelLayerTree;
+
+/**
+ * Editor that displays the compositing actors in the level and displays the layer tree and layer properties for them
+ */
+class SCompositeEditorPanel : public SCompoundWidget
+{
+public:
+	/** Registers the tab spawner for the panel with the editor */
+	static void RegisterTabSpawner();
+
+	/** Unregisters the tab spawner for the panel with the editor */
+	static void UnregisterTabSpawner();
+	
+public:
+	SLATE_BEGIN_ARGS(SCompositeEditorPanel) { }
+	SLATE_END_ARGS()
+	
+	void Construct(const FArguments& InArgs);
+
+private:
+	/** Raised when the selected layers in the tree view have changed */
+	void OnLayerSelectionChanged(const TArray<UObject*>& SelectedLayers);
+	
+private:
+	/** The tree view widget that displays the composite layers */
+	TSharedPtr<SCompositePanelLayerTree> LayerTree;
+
+	/** The details view that displays the properties of the selected layers and actors */
+	TSharedPtr<IDetailsView> DetailsView;
+};
diff --git a/Engine/Plugins/Experimental/Compositing/Composite/Source/CompositeEditor/Private/UI/SCompositePanelLayerTree.cpp b/Engine/Plugins/Experimental/Compositing/Composite/Source/CompositeEditor/Private/UI/SCompositePanelLayerTree.cpp
new file mode 100644
--- /dev/null
+++ b/Engine/Plugins/Experimental/Compositing/Composite/Source/CompositeEditor/Private/UI/SCompositePanelLayerTree.cpp
@@ -0,0 +1,1795 @@
+﻿// Copyright Epic Games, Inc. All Rights Reserved.
+
+#include "SCompositePanelLayerTree.h"
+
+#include "AudioMixerBlueprintLibrary.h"
+#include "CompositeActor.h"
+#include "Editor.h"
+#include "EngineUtils.h"
+#include "Factories.h"
+#include "LevelEditorSubsystem.h"
+#include "ScopedTransaction.h"
+#include "SPositiveActionButton.h"
+#include "UnrealExporter.h"
+#include "DragAndDrop/DecoratedDragDropOp.h"
+#include "Exporters/Exporter.h"
+#include "Filters/SBasicFilterBar.h"
+#include "Framework/Commands/GenericCommands.h"
+#include "HAL/PlatformApplicationMisc.h"
+#include "Layers/CompositeLayerBase.h"
+#include "Styling/SlateIconFinder.h"
+#include "Styling/ToolBarStyle.h"
+#include "UObject/Class.h"
+#include "Widgets/Text/SInlineEditableTextBlock.h"
+#include "Widgets/Views/STableRow.h"
+#include "Widgets/Views/STreeView.h"
+#include "Windows/WindowsPlatformApplicationMisc.h"
+
+#define LOCTEXT_NAMESPACE "SCompositePanelLayerTree"
+
+namespace CompositePanelLayerTree
+{
+	static const FName LayerTreeColumn_Enabled = "Enabled";
+	static const FName LayerTreeColumn_Solo = "Solo";
+	static const FName LayerTreeColumn_Name = "Name";
+	static const FName LayerTreeColumn_Type = "Type";
+	static const FName LayerTreeColumn_Operation = "Operation";
+}
+
+/** Toolbar widget that contains filtering, the add button, and the settings button */
+class SCompositePanelLayerTreeToolbar : public SCompoundWidget
+{
+private:
+	using FFilterType = SCompositePanelLayerTree::FCompositeActorTreeItemPtr;
+	using FCompositeLayersFilter = TTextFilter<FFilterType>;
+	
+public:
+	DECLARE_DELEGATE_OneParam(FOnPilotCameraToggled, bool);
+	DECLARE_DELEGATE_OneParam(FOnLayerAdded, const UClass*)
+	
+	SLATE_BEGIN_ARGS(SCompositePanelLayerTreeToolbar) {}
+		SLATE_ATTRIBUTE(bool, CanPilotCamera)
+		SLATE_ATTRIBUTE(bool, IsPilotingCamera)
+		SLATE_EVENT(FOnPilotCameraToggled, OnPilotCameraToggled)
+		SLATE_EVENT(FSimpleDelegate, OnFilterChanged)
+		SLATE_EVENT(FOnLayerAdded, OnLayerAdded)
+	SLATE_END_ARGS()
+
+	void Construct(const FArguments& InArgs)
+	{
+		IsPilotingCamera = InArgs._IsPilotingCamera;
+		OnPilotCameraToggled = InArgs._OnPilotCameraToggled;
+		OnFilterChanged = InArgs._OnFilterChanged;
+		OnLayerAdded = InArgs._OnLayerAdded;
+
+		InitializeFilters();
+		
+		ChildSlot
+		[
+			SNew(SVerticalBox)
+
+			+SVerticalBox::Slot()
+			.AutoHeight()
+			[
+				SNew(SHorizontalBox)
+
+				+SHorizontalBox::Slot()
+				.AutoWidth()
+				.Padding(0.0f, 0.0f, 4.0f, 0.0f)
+				[
+					SNew(SPositiveActionButton)
+					.Icon(FAppStyle::Get().GetBrush("Icons.Plus"))
+					.Text(LOCTEXT("AddMediaButtonLabel", "Add"))
+					.ToolTipText(LOCTEXT("AddMediaButtonToolTip", "Add new composite layer to the selected composite actor"))
+					.OnGetMenuContent(this, &SCompositePanelLayerTreeToolbar::GetAddMenuContent)
+				]
+
+				+SHorizontalBox::Slot()
+				.AutoWidth()
+				.Padding(4.f, 0.f, 4.f, 0.f)
+				[
+					SNew(SCheckBox)
+					.Padding(4.0)
+					.Style(&FAppStyle::Get().GetWidgetStyle<FToolBarStyle>("ViewportToolbar.Raised").ToggleButton)
+					.Type(ESlateCheckBoxType::ToggleButton)
+					.ToolTipText(LOCTEXT("PilotButtonToolTip", "Pilot the camera actor associated with the selected composite actor"))
+					.IsEnabled(InArgs._CanPilotCamera)
+					.IsChecked_Lambda([this]
+					{
+						return IsPilotingCamera.Get(false) ? ECheckBoxState::Checked : ECheckBoxState::Unchecked;
+					})
+					.OnCheckStateChanged_Lambda([this](ECheckBoxState CheckState)
+					{
+						OnPilotCameraToggled.ExecuteIfBound(CheckState == ECheckBoxState::Checked);
+					})
+					[
+						SNew(SImage)
+						.Image(FAppStyle::GetBrush("LevelViewport.PilotSelectedActor"))
+						.ColorAndOpacity(FSlateColor::UseForeground())
+					]
+				]
+				
+				+SHorizontalBox::Slot()
+				.FillWidth(1.0f)
+				[
+					SNew(SFilterSearchBox)
+					.HintText(LOCTEXT("FilterSearch", "Search..."))
+					.ToolTipText(LOCTEXT("FilterSearchHint", "Type here to search for specific composite layers"))
+					.OnTextChanged_Lambda([this](const FText& InText)
+					{
+						TextFilter->SetRawFilterText(InText);
+					})
+				]
+
+				+SHorizontalBox::Slot()
+				.AutoWidth()
+				.Padding(4.f, 0.f, 0.f, 0.f)
+				[
+					SNew(SComboButton)
+					.ComboButtonStyle(FAppStyle::Get(), "SimpleComboButtonWithIcon")
+					.HasDownArrow(false)
+					.ButtonContent()
+					[
+						SNew(SImage)
+						.ColorAndOpacity(FSlateColor::UseForeground())
+						.Image(FAppStyle::Get().GetBrush("Icons.Settings"))
+					]
+				]
+			]
+		];
+	}
+
+	/** Gets whether the specified tree item passes any filters enabled in the toolbar */
+	bool ItemPassesFilters(const FFilterType& InItem) const
+	{
+		return TextFilter->PassesFilter(InItem);
+	}
+
+	/** Clears the toolbar filters */
+	void ClearFilter()
+	{
+		TextFilter->SetRawFilterText(FText::GetEmpty());
+	}
+	
+private:
+	/** Configures any filters for the toolbar */
+	void InitializeFilters()
+	{
+		TextFilter = MakeShared<FCompositeLayersFilter>(FCompositeLayersFilter::FItemToStringArray::CreateLambda([](const FFilterType& InItem, TArray<FString>& OutStrings)
+		{
+			if (!InItem.IsValid())
+			{
+				return;
+			}
+
+			if (!InItem->HasValidCompositeActor())
+			{
+				return;
+			}
+
+			OutStrings.Add(InItem->CompositeActor->GetActorNameOrLabel());
+			OutStrings.Add(InItem->CompositeActor->GetClass()->GetDisplayNameText().ToString());
+				
+			if (UCompositeLayerBase* CompositeLayer = InItem->GetCompositeLayer())
+			{
+				OutStrings.Add(CompositeLayer->Name);
+				OutStrings.Add(CompositeLayer->GetClass()->GetDisplayNameText().ToString());
+			}
+		}));
+		
+		TextFilter->OnChanged().AddSP(this, &SCompositePanelLayerTreeToolbar::FilterChanged);
+	}
+
+	/** Gets the menu content of the Add button, which will allow users to pick which kind of layer they want to add */
+	TSharedRef<SWidget> GetAddMenuContent()
+	{
+		constexpr bool bCloseMenuAfterSelection = true;
+		FMenuBuilder MenuBuilder(bCloseMenuAfterSelection, nullptr);
+		
+		TArray<UClass*> BaseLayerTypes;
+		const bool bRecursive = false;
+		GetDerivedClasses(UCompositeLayerBase::StaticClass(), BaseLayerTypes, bRecursive);
+
+		auto ShouldSkipClass = [](const UClass* InClass)
+		{
+			constexpr EClassFlags InvalidClassFlags = CLASS_Deprecated | CLASS_NewerVersionExists | CLASS_Abstract;
+			return InClass->HasAnyClassFlags(InvalidClassFlags) ||
+				InClass->GetName().StartsWith(TEXT("SKEL_")) ||
+				InClass->GetName().StartsWith(TEXT("REINST_"));
+		};
+		
+		// Add the native child classes of UCompositeLayerBase to the menu first
+		for (const UClass* BaseLayerType : BaseLayerTypes)
+		{
+			if (ShouldSkipClass(BaseLayerType))
+			{
+				continue;
+			}
+			
+			MenuBuilder.AddMenuEntry(
+				BaseLayerType->GetDisplayNameText(),
+				TAttribute<FText>(),
+				FSlateIconFinder::FindIconForClass(BaseLayerType),
+				FUIAction(FExecuteAction::CreateSP(this, &SCompositePanelLayerTreeToolbar::AddLayer, BaseLayerType)));
+		}
+				
+		// Add any blueprint derived classes in a new section of the menu
+		bool bHasAddedSeparator = false;
+		for (const UClass* BaseLayerType : BaseLayerTypes)
+		{
+			TArray<UClass*> ChildLayerTypes;
+			GetDerivedClasses(BaseLayerType, ChildLayerTypes);
+
+			for (const UClass* ChildLayerType : ChildLayerTypes)
+			{
+				if (ShouldSkipClass(ChildLayerType))
+				{
+					continue;
+				}
+				
+				if (!bHasAddedSeparator)
+				{
+					MenuBuilder.AddSeparator();
+					bHasAddedSeparator = true;
+				}
+
+				MenuBuilder.AddMenuEntry(
+					ChildLayerType->GetDisplayNameText(),
+					TAttribute<FText>(),
+					FSlateIconFinder::FindIconForClass(ChildLayerType),
+					FUIAction(FExecuteAction::CreateSP(this, &SCompositePanelLayerTreeToolbar::AddLayer, ChildLayerType)));
+			}
+		}
+		
+		return MenuBuilder.MakeWidget();
+	}
+
+	/** Callback raised when a layer type is selected in the Add dropdown menu */
+	void AddLayer(const UClass* InLayerType)
+	{
+		OnLayerAdded.ExecuteIfBound(InLayerType);
+	}
+
+	/** Callback raised when the toolbar's active filter has changed */
+	void FilterChanged()
+	{
+		OnFilterChanged.ExecuteIfBound();
+	}
+	
+private:
+	/** Attribute used to query if a camera actor is actively being piloted in the editor */
+	TAttribute<bool> IsPilotingCamera;
+
+	/** Manages the text filter that has been entered into the filter text box */
+	TSharedPtr<FCompositeLayersFilter> TextFilter;
+
+	FOnPilotCameraToggled OnPilotCameraToggled;
+	FSimpleDelegate OnFilterChanged;
+	FOnLayerAdded OnLayerAdded;
+};
+
+/** Row widget used to display a tree item in the composite layer tree */
+class SCompositePanelLayerTreeItemRow : public SMultiColumnTableRow<SCompositePanelLayerTree::FCompositeActorTreeItemPtr>
+{
+	using FTreeItem = SCompositePanelLayerTree::FCompositeActorTreeItem;
+	using FTreeItemPtr = SCompositePanelLayerTree::FCompositeActorTreeItemPtr;
+	using FTreeItemWeakPtr = TWeakPtr<FTreeItem>;
+	
+public:
+	DECLARE_DELEGATE_RetVal_OneParam(FReply, FOnLayerSoloToggled, const FTreeItemPtr&)
+	DECLARE_DELEGATE_TwoParams(FOnLayerMoved, const FTreeItemPtr& /* InItemToMove */, int32 /* InDestIndex */)
+
+private:
+	/** Drag/drop operation for changing the order of sources or outputs in the tree view */
+	class FCompositeLayerDragDropOp : public FDecoratedDragDropOp
+	{
+	public:
+		DRAG_DROP_OPERATOR_TYPE(FCompositeLayerDragDropOp, FDecoratedDragDropOp)
+
+		static TSharedRef<FCompositeLayerDragDropOp> New(const FTreeItemPtr& InDraggedItem)
+		{
+			TSharedRef<FCompositeLayerDragDropOp> Operation = MakeShared<FCompositeLayerDragDropOp>();
+			Operation->DraggedItem = InDraggedItem;
+			Operation->Construct();
+			return Operation;
+		}
+
+		virtual void Construct() override
+		{
+			if (FTreeItemPtr PinnedItem = DraggedItem.Pin())
+			{
+				if (UCompositeLayerBase* CompositeLayer = PinnedItem->GetCompositeLayer())
+				{
+					const FString LayerName = !CompositeLayer->Name.IsEmpty() ? CompositeLayer->Name : CompositeLayer->GetName();
+					const FSlateBrush* LayerIcon = FSlateIconFinder::FindIconForClass(CompositeLayer->GetClass()).GetIcon();
+					SetToolTip(FText::FromString(LayerName), LayerIcon);
+				}
+			}
+			
+			FDecoratedDragDropOp::Construct();
+		}
+
+		FTreeItemWeakPtr DraggedItem;
+	};
+	
+public:
+	SLATE_BEGIN_ARGS(SCompositePanelLayerTreeItemRow) { }
+		SLATE_EVENT(FOnLayerSoloToggled, OnLayerSoloToggled)
+		SLATE_EVENT(FOnLayerMoved, OnLayerMoved)
+	SLATE_END_ARGS()
+
+	void Construct(const FArguments& InArgs, const TSharedRef<STableViewBase>& InOwnerTable, FTreeItemPtr InTreeItem)
+	{
+		TreeItem = InTreeItem;
+		OnLayerSoloToggled = InArgs._OnLayerSoloToggled;
+		OnLayerMoved = InArgs._OnLayerMoved;
+
+		STableRow<FTreeItemPtr>::FArguments Args = FSuperRowType::FArguments()
+           .Style(&FAppStyle::Get().GetWidgetStyle<FTableRowStyle>("SceneOutliner.TableViewRow"))
+           .OnDragDetected(this, &SCompositePanelLayerTreeItemRow::HandleDragDetected)
+           .OnCanAcceptDrop(this, &SCompositePanelLayerTreeItemRow::HandleCanAcceptDrop)
+           .OnAcceptDrop(this, &SCompositePanelLayerTreeItemRow::HandleAcceptDrop);
+		
+		SMultiColumnTableRow::Construct(Args, InOwnerTable);
+	}
+
+	/** Sets the name text block to be in edit mode for this row widget */
+	void RequestRename()
+	{
+		if (NameTextBlock.IsValid())
+		{
+			NameTextBlock->EnterEditingMode();
+		}
+	}
+	
+private:
+	virtual TSharedRef<SWidget> GenerateWidgetForColumn(const FName& InColumnName) override
+	{
+		if (InColumnName == CompositePanelLayerTree::LayerTreeColumn_Enabled)
+		{
+			return SNew(SCheckBox)
+				.IsChecked(this, &SCompositePanelLayerTreeItemRow::IsItemEnabled)
+				.OnCheckStateChanged(this, &SCompositePanelLayerTreeItemRow::OnIsItemEnabledChanged);
+		}
+		if (InColumnName == CompositePanelLayerTree::LayerTreeColumn_Solo)
+		{
+			return SNew(SImage).Image(FAppStyle::GetBrush("MediaAsset.AssetActions.Solo.Small"))
+				.ColorAndOpacity_Raw(this, &SCompositePanelLayerTreeItemRow::GetItemSoloIconColor)
+				.OnMouseButtonDown(this, &SCompositePanelLayerTreeItemRow::OnItemSoloIconMouseButtonDown);
+		}
+		if (InColumnName == CompositePanelLayerTree::LayerTreeColumn_Name)
+		{
+			return SNew(SHorizontalBox)
+
+				+SHorizontalBox::Slot()
+				.AutoWidth()
+				.VAlign(VAlign_Center)
+				.Padding(6.f, 0.f, 0.f, 0.f)
+				[
+					SNew(SExpanderArrow, SharedThis(this)).IndentAmount(24)
+				]
+				
+				+SHorizontalBox::Slot()
+				.AutoWidth()
+				.VAlign(VAlign_Center)
+				.Padding(0.0f, 1.0f, 6.0f, 1.0f)
+				[
+					SNew(SBox)
+					.WidthOverride(16.0f)
+					.HeightOverride(16.0f)
+					[
+						SNew(SImage).Image(this, &SCompositePanelLayerTreeItemRow::GetItemIcon)
+					]
+				]
+
+				+SHorizontalBox::Slot()
+				.VAlign(VAlign_Center)
+				[
+					SAssignNew(NameTextBlock, SInlineEditableTextBlock)
+					.Text(this, &SCompositePanelLayerTreeItemRow::GetItemName)
+					.OnTextCommitted(this, &SCompositePanelLayerTreeItemRow::OnItemNameChanged)
+					.IsSelected(FIsSelected::CreateSP(this, &SCompositePanelLayerTreeItemRow::IsSelectedExclusively))
+				];
+		}
+		if (InColumnName == CompositePanelLayerTree::LayerTreeColumn_Type)
+		{
+			return SNew(STextBlock).Text(this, &SCompositePanelLayerTreeItemRow::GetItemType);
+		}
+		if (InColumnName == CompositePanelLayerTree::LayerTreeColumn_Operation)
+		{
+			return SNew(STextBlock).Text(this, &SCompositePanelLayerTreeItemRow::GetItemOperation);	
+		}
+		
+		return SNullWidget::NullWidget;
+	}
+
+	ACompositeActor* GetCompositeActor() const
+	{
+		return TreeItem->HasValidCompositeActor() ? TreeItem->CompositeActor.Get() : nullptr;
+	}
+
+	/** Gets the checkbox state for the row item's Enabled flag */
+	ECheckBoxState IsItemEnabled() const
+	{
+		if (TreeItem->HasValidCompositeLayer())
+		{
+			if (UCompositeLayerBase* CompositeLayer = TreeItem->GetCompositeLayer())
+			{
+				return CompositeLayer->bIsEnabled ? ECheckBoxState::Checked : ECheckBoxState::Unchecked;
+			}
+			
+			// In the case where there is a valid layer index, but the layer itself is null
+			return ECheckBoxState::Unchecked;
+		}
+
+		if (ACompositeActor* CompositeActor = GetCompositeActor())
+		{
+			return CompositeActor->IsEnabled() ? ECheckBoxState::Checked : ECheckBoxState::Unchecked;
+		}
+		
+		return ECheckBoxState::Unchecked;
+	}
+
+	/** Raised when the row item's Enabled checkbox has changed state */
+	void OnIsItemEnabledChanged(ECheckBoxState InCheckBoxState)
+	{
+		if (TreeItem->HasValidCompositeLayer())
+		{
+			if (UCompositeLayerBase* CompositeLayer = TreeItem->GetCompositeLayer())
+			{
+				FScopedTransaction SetEnabledTransaction(LOCTEXT("SetEnabledTransaction", "Set Enabled"));
+				CompositeLayer->Modify();
+
+				// TODO: Should convert to a setter that can ensure ActivateSceneCapture is called when re-enabling
+				CompositeLayer->bIsEnabled = InCheckBoxState == ECheckBoxState::Checked;
+				return;
+			}
+
+			// In the case where there is a valid layer index, but the layer itself is null
+			return;
+		}
+		
+		if (ACompositeActor* CompositeActor = GetCompositeActor())
+		{
+			FScopedTransaction SetEnabledTransaction(LOCTEXT("SetEnabledTransaction", "Set Enabled"));
+			CompositeActor->Modify();
+
+			CompositeActor->SetEnabled(InCheckBoxState == ECheckBoxState::Checked);
+			return;
+		}
+	}
+
+	/** Gets the color for the row item's Solo icon */
+	FSlateColor GetItemSoloIconColor() const
+	{
+		bool bIsLayer = false;
+		bool bIsLayerSolo = false;
+		if (UCompositeLayerBase* CompositeLayer = TreeItem->GetCompositeLayer())
+		{
+			bIsLayer = true;
+			bIsLayerSolo = CompositeLayer->bIsSolo;
+		}
+		
+		if (!bIsLayer)
+		{
+			return FLinearColor::Transparent;
+		}
+
+		if (IsHovered())
+		{
+			return FStyleColors::ForegroundHover;
+		}
+
+		return bIsLayerSolo ? FSlateColor::UseForeground() : FLinearColor::Transparent;
+	}
+
+	/** Raised when the Solo button has been pressed */
+	FReply OnItemSoloIconMouseButtonDown(const FGeometry& Geometry, const FPointerEvent& PointerEvent)
+	{
+		if (OnLayerSoloToggled.IsBound())
+		{
+			return OnLayerSoloToggled.Execute(TreeItem);
+		}
+		
+		return FReply::Unhandled();
+	}
+
+	/** Gets the visibility of the row item's Solo button */
+	EVisibility GetItemSoloVisibility() const
+	{
+		bool bIsLayerSolo = false;
+		if (UCompositeLayerBase* CompositeLayer = TreeItem->GetCompositeLayer())
+		{
+			bIsLayerSolo = CompositeLayer->bIsSolo;
+		}
+		
+		return (IsHighlighted() || bIsLayerSolo) ? EVisibility::Visible : EVisibility::Hidden;
+	}
+
+	/** Gets the row item's icon to display next to its name */
+	const FSlateBrush* GetItemIcon() const
+	{
+		if (TreeItem->HasValidCompositeLayer())
+		{
+			if (UCompositeLayerBase* CompositeLayer = TreeItem->GetCompositeLayer())
+			{
+				return FSlateIconFinder::FindIconForClass(CompositeLayer->GetClass()).GetIcon();
+			}
+			
+			// In the case where there is a valid layer index, but the layer itself is null
+			return nullptr;
+		}
+		
+		if (ACompositeActor* CompositeActor = GetCompositeActor())
+		{
+			return FSlateIconFinder::FindIconForClass(CompositeActor->GetClass()).GetIcon(); 
+		}
+
+		return nullptr;
+	}
+
+	/** Gets the row item's display name */
+	FText GetItemName() const
+	{
+		if (TreeItem->HasValidCompositeLayer())
+		{
+			if (UCompositeLayerBase* CompositeLayer = TreeItem->GetCompositeLayer())
+			{
+				if (CompositeLayer->Name.IsEmpty())
+				{
+					// Display the UObject name for the layer if its Name property is empty.
+					return FText::FromString(CompositeLayer->GetName());
+				}
+			
+				return FText::FromString(CompositeLayer->Name);
+			}
+
+			// In the case where there is a valid layer index, but the layer itself is null
+			return LOCTEXT("LayerNotConfiguredLabel", "Layer not configured");
+		}
+
+		if (ACompositeActor* CompositeActor = GetCompositeActor())
+		{
+			return FText::FromString(CompositeActor->GetActorNameOrLabel());
+		}
+
+		return FText::GetEmpty();
+	}
+
+	/** Raised when the row item's name has been changed */
+	void OnItemNameChanged(const FText& InNewName, ETextCommit::Type InCommitType)
+	{
+		if (TreeItem->HasValidCompositeLayer())
+		{
+			if (UCompositeLayerBase* CompositeLayer = TreeItem->GetCompositeLayer())
+			{
+				FScopedTransaction SetLayerNameTransaction(LOCTEXT("SetLayerNameTransaction", "Set Layer Name"));
+				CompositeLayer->Modify();
+
+				CompositeLayer->Name = InNewName.ToString();
+				return;
+			}
+
+			// In the case where there is a valid layer index, but the layer itself is null
+			return;
+		}
+
+		if (ACompositeActor* CompositeActor = GetCompositeActor())
+		{
+			FScopedTransaction RenameActorTransaction(LOCTEXT("RenameActorTransaction", "Rename Composite Actor"));
+			FActorLabelUtilities::RenameExistingActor(CompositeActor, InNewName.ToString(), true);
+		}
+	}
+
+	/** Gets the display text of the row item's type */
+	FText GetItemType() const
+	{
+		if (TreeItem->HasValidCompositeLayer())
+		{
+			if (UCompositeLayerBase* CompositeLayer = TreeItem->GetCompositeLayer())
+			{
+				return CompositeLayer->GetClass()->GetDisplayNameText();
+			}
+
+			return LOCTEXT("NoneTypeLabel", "None");
+		}
+
+		if (ACompositeActor* CompositeActor = GetCompositeActor())
+		{
+			return CompositeActor->GetClass()->GetDisplayNameText();
+		}
+
+		return FText::GetEmpty();
+	}
+
+	/** Gets the display text of the row item's compositing operation */
+	FText GetItemOperation() const
+	{
+		if (UCompositeLayerBase* CompositeLayer = TreeItem->GetCompositeLayer())
+		{
+			return StaticEnum<ECompositeCoreMergeOp>()->GetDisplayNameTextByValue((int64)CompositeLayer->Operation);
+		}
+
+		return LOCTEXT("NullOperationLabel", "-");
+	}
+
+	/** Raised when a drag start has been detected on this row item */
+	FReply HandleDragDetected(const FGeometry& Geometry, const FPointerEvent& MouseEvent)
+	{
+		if (!TreeItem.IsValid())
+		{
+			return FReply::Unhandled();
+		}
+
+		if (!TreeItem->HasValidCompositeLayer())
+		{
+			return FReply::Unhandled();
+		}
+		
+		if (MouseEvent.IsMouseButtonDown(EKeys::LeftMouseButton))
+		{
+			TSharedPtr<FDragDropOperation> DragDropOp = FCompositeLayerDragDropOp::New(TreeItem);
+			return FReply::Handled().BeginDragDrop(DragDropOp.ToSharedRef());
+		}
+
+		return FReply::Unhandled();
+	}
+
+	/** Raised when the user is attempting to drop something onto this item */
+	TOptional<EItemDropZone> HandleCanAcceptDrop(const FDragDropEvent& InDragDropEvent, EItemDropZone InItemDropZone, FTreeItemPtr InTreeItem)
+	{
+		TSharedPtr<FCompositeLayerDragDropOp> DragDropOp = InDragDropEvent.GetOperationAs<FCompositeLayerDragDropOp>();
+
+		if (!DragDropOp.IsValid() || !DragDropOp->DraggedItem.IsValid())
+		{
+			return TOptional<EItemDropZone>();
+		}
+
+		FTreeItemPtr DraggedItem = DragDropOp->DraggedItem.Pin();
+		if (DraggedItem->HasValidCompositeActor() && DraggedItem->CompositeActor.Get() == InTreeItem->CompositeActor.Get())
+		{
+			if (!InTreeItem->HasValidCompositeLayer())
+			{
+				return EItemDropZone::BelowItem;
+			}
+
+			return InItemDropZone == EItemDropZone::OntoItem ? EItemDropZone::AboveItem : InItemDropZone;
+		}
+		
+		return TOptional<EItemDropZone>();
+	}
+
+	/** Raised when the user has dropped something onto this item */
+	FReply HandleAcceptDrop(const FDragDropEvent& InDragDropEvent, EItemDropZone InItemDropZone, FTreeItemPtr InTreeItem)
+	{
+		TSharedPtr<FCompositeLayerDragDropOp> DragDropOp = InDragDropEvent.GetOperationAs<FCompositeLayerDragDropOp>();
+
+		if (!DragDropOp.IsValid() || !DragDropOp->DraggedItem.IsValid())
+		{
+			return FReply::Unhandled();
+		}
+
+		if (OnLayerMoved.IsBound())
+		{
+			int32 DestIndex = InItemDropZone == EItemDropZone::AboveItem ? InTreeItem->LayerIndex: InTreeItem->LayerIndex + 1;
+
+			// If the destination is further down than the item being moved, we must subtract one from the destination index
+			// to account for the removal of the original item
+			if (InTreeItem->LayerIndex > DragDropOp->DraggedItem.Pin()->LayerIndex)
+			{
+				--DestIndex;
+			}
+			
+			OnLayerMoved.Execute(DragDropOp->DraggedItem.Pin(), DestIndex);
+			return FReply::Handled();
+		}
+		
+		return FReply::Unhandled();
+	}
+	
+private:
+	/** The tree item this widget is outputting */
+	FTreeItemPtr TreeItem;
+
+	/** Text block widget used to display the tree item's name */
+	TSharedPtr<SInlineEditableTextBlock> NameTextBlock;
+	
+	/** Callback that is raised when the layer's "solo" flag is enabled */
+	FOnLayerSoloToggled OnLayerSoloToggled;
+
+	/** Callback that is raised when a layer is moved via a drag and drop operation */
+	FOnLayerMoved OnLayerMoved;
+};
+
+bool SCompositePanelLayerTree::FCompositeActorTreeItem::HasValidCompositeActor() const
+{
+	return CompositeActor.IsValid();
+}
+
+bool SCompositePanelLayerTree::FCompositeActorTreeItem::HasValidCompositeLayer() const
+{
+	return CompositeActor.IsValid() && CompositeActor->CompositeLayers.IsValidIndex(LayerIndex);
+}
+
+UCompositeLayerBase* SCompositePanelLayerTree::FCompositeActorTreeItem::GetCompositeLayer() const
+{
+	return HasValidCompositeLayer() ? CompositeActor->CompositeLayers[LayerIndex] : nullptr;
+}
+
+void SCompositePanelLayerTree::Construct(const FArguments& InArgs)
+{
+	OnSelectionChanged = InArgs._OnSelectionChanged;
+
+	if (GEditor)
+	{
+		GEditor->RegisterForUndo(this);
+	}
+
+	if (GEngine)
+	{
+		GEngine->OnLevelActorAdded().AddSP(this, &SCompositePanelLayerTree::OnActorAddedToLevel);
+		GEngine->OnLevelActorListChanged().AddSP(this, &SCompositePanelLayerTree::OnLevelActorListChanged);
+		GEngine->OnLevelActorDeleted().AddSP(this, &SCompositePanelLayerTree::OnActorRemovedFromLevel);
+	}
+	
+	FEditorDelegates::MapChange.AddSP(this, &SCompositePanelLayerTree::OnMapChange);
+	FEditorDelegates::NewCurrentLevel.AddSP(this, &SCompositePanelLayerTree::OnNewCurrentLevel);
+	FWorldDelegates::LevelAddedToWorld.AddSP(this, &SCompositePanelLayerTree::OnLevelAdded);
+	FWorldDelegates::LevelRemovedFromWorld.AddSP(this, &SCompositePanelLayerTree::OnLevelRemoved);
+	FCoreUObjectDelegates::OnObjectsReplaced.AddSP(this, &SCompositePanelLayerTree::OnObjectReplaced);
+	FCoreUObjectDelegates::OnObjectPropertyChanged.AddSP(this, &SCompositePanelLayerTree::OnObjectPropertyChanged);
+	
+	CommandList = MakeShared<FUICommandList>();
+	BindCommands();
+	
+	ChildSlot
+	[
+		SNew(SVerticalBox)
+
+		+SVerticalBox::Slot()
+		.AutoHeight()
+		.Padding(8.0f, 8.0f, 8.0f, 4.0f)
+		[
+			SAssignNew(Toolbar, SCompositePanelLayerTreeToolbar)
+			.CanPilotCamera(this, &SCompositePanelLayerTree::CanPilotCamera)
+			.IsPilotingCamera(this, &SCompositePanelLayerTree::IsPilotingCamera)
+			.OnPilotCameraToggled(this, &SCompositePanelLayerTree::OnPilotCameraToggled)
+			.OnFilterChanged(this, &SCompositePanelLayerTree::OnFilterChanged)
+			.OnLayerAdded(this, &SCompositePanelLayerTree::OnLayerAdded)
+		]
+
+		+SVerticalBox::Slot()
+		[
+			SAssignNew(TreeView, STreeView<FCompositeActorTreeItemPtr>)
+			.TreeItemsSource(&FilteredCompositeActorTreeItems)
+			.HeaderRow(
+				SNew(SHeaderRow)
+
+				// TODO: When multiple composite actors are supported, add a column here to allow a specific composite actor to be set to the active one
+				
+				+SHeaderRow::Column(CompositePanelLayerTree::LayerTreeColumn_Enabled)
+				.DefaultLabel(FText::GetEmpty())
+				.FixedWidth(24.0f)
+				.HAlignHeader(HAlign_Center)
+				.VAlignHeader(VAlign_Center)
+				.HAlignCell(HAlign_Center)
+				.VAlignCell(VAlign_Center)
+				[
+					SNew(SCheckBox)
+					.IsChecked(this, &SCompositePanelLayerTree::GetGlobalEnabledCheckState)
+					.OnCheckStateChanged(this, &SCompositePanelLayerTree::OnGlobalEnabledCheckStateChanged)
+				]
+				
+				+SHeaderRow::Column(CompositePanelLayerTree::LayerTreeColumn_Solo)
+				.DefaultLabel(FText::GetEmpty())
+				.FixedWidth(24.0f)
+				.HAlignHeader(HAlign_Center)
+				.VAlignHeader(VAlign_Center)
+				.HAlignCell(HAlign_Center)
+				.VAlignCell(VAlign_Center)
+				[
+					SNew(SImage).Image(FAppStyle::GetBrush("MediaAsset.AssetActions.Solo.Small"))
+				]
+				
+				+SHeaderRow::Column(CompositePanelLayerTree::LayerTreeColumn_Name)
+				.DefaultLabel(LOCTEXT("LayerTreeColumnLabel", "Layer"))
+				.FillWidth(0.55)
+
+				+SHeaderRow::Column(CompositePanelLayerTree::LayerTreeColumn_Type)
+				.DefaultLabel(LOCTEXT("TypeTreeColumnLabel", "Type"))
+				.FillWidth(0.3)
+
+				+SHeaderRow::Column(CompositePanelLayerTree::LayerTreeColumn_Operation)
+				.DefaultLabel(LOCTEXT("OperationTreeColumnLabel", "Operation"))
+				.FillWidth(0.15)
+			)
+			.OnGenerateRow_Lambda([this](FCompositeActorTreeItemPtr InTreeItem, const TSharedRef<STableViewBase>& InOwnerTable)
+			{
+				return SNew(SCompositePanelLayerTreeItemRow, InOwnerTable, InTreeItem)
+					.OnLayerSoloToggled(this, &SCompositePanelLayerTree::OnLayerSoloToggled)
+					.OnLayerMoved(this, &SCompositePanelLayerTree::OnLayerMoved);
+			})
+			.OnGetChildren_Lambda([](FCompositeActorTreeItemPtr InTreeItem, TArray<FCompositeActorTreeItemPtr>& OutChildren)
+			{
+				if (InTreeItem.IsValid())
+				{
+					for (FCompositeActorTreeItemPtr& Child : InTreeItem->Children)
+					{
+						if (Child.IsValid() && !Child->bFilteredOut)
+						{
+							OutChildren.Add(Child);
+						}
+					}
+				}
+			})
+			.OnSelectionChanged(this, &SCompositePanelLayerTree::OnLayerSelectionChanged)
+			.OnContextMenuOpening(this, &SCompositePanelLayerTree::CreateTreeContextMenu)
+		]
+
+		+SVerticalBox::Slot()
+		.AutoHeight()
+		[
+			SNew(SBorder)
+			.BorderImage(FAppStyle::Get().GetBrush("Brushes.Header"))
+			.VAlign(VAlign_Center)
+			.HAlign(HAlign_Left)
+			.Padding(FMargin(14, 9))
+			[
+				SNew( STextBlock )
+				.Text(this, &SCompositePanelLayerTree::GetFilterStatusText)
+				.ColorAndOpacity(this, &SCompositePanelLayerTree::GetFilterStatusTextColor)
+			]
+		]
+	];
+
+	FillCompositeActorTree();
+}
+
+SCompositePanelLayerTree::~SCompositePanelLayerTree()
+{
+	FEditorDelegates::MapChange.RemoveAll(this);
+	FEditorDelegates::NewCurrentLevel.RemoveAll(this);
+	FWorldDelegates::LevelAddedToWorld.RemoveAll(this);
+	FWorldDelegates::LevelRemovedFromWorld.RemoveAll(this);
+	FCoreUObjectDelegates::OnObjectsReplaced.RemoveAll(this);
+	FCoreUObjectDelegates::OnObjectPropertyChanged.RemoveAll(this);
+	
+	if (GEngine)
+	{
+		GEngine->OnLevelActorAdded().RemoveAll(this);
+		GEngine->OnLevelActorListChanged().RemoveAll(this);
+		GEngine->OnLevelActorDeleted().RemoveAll(this);
+	}
+	
+	if (GEditor)
+	{
+		GEditor->UnregisterForUndo(this);
+	}
+}
+
+FReply SCompositePanelLayerTree::OnKeyDown(const FGeometry& MyGeometry, const FKeyEvent& InKeyEvent)
+{
+	if (this->CommandList->ProcessCommandBindings(InKeyEvent))
+	{
+		return FReply::Handled();
+	}
+	
+	return FReply::Unhandled();
+}
+
+void SCompositePanelLayerTree::PostUndo(bool bSuccess)
+{
+	FillCompositeActorTree();
+}
+
+void SCompositePanelLayerTree::PostRedo(bool bSuccess)
+{
+	PostUndo(bSuccess);
+}
+
+void SCompositePanelLayerTree::BindCommands()
+{
+	CommandList->MapAction(FGenericCommands::Get().Copy,
+		FExecuteAction::CreateSP(this, &SCompositePanelLayerTree::CopySelectedItems),
+		FCanExecuteAction::CreateSP(this, &SCompositePanelLayerTree::CanCopySelectedItems));
+
+	CommandList->MapAction(FGenericCommands::Get().Cut,
+		FExecuteAction::CreateSP(this, &SCompositePanelLayerTree::CutSelectedItems),
+		FCanExecuteAction::CreateSP(this, &SCompositePanelLayerTree::CanCutSelectedItems));
+
+	CommandList->MapAction(FGenericCommands::Get().Paste,
+		FExecuteAction::CreateSP(this, &SCompositePanelLayerTree::PasteSelectedItems),
+		FCanExecuteAction::CreateSP(this, &SCompositePanelLayerTree::CanPasteSelectedItems));
+
+	CommandList->MapAction(FGenericCommands::Get().Duplicate,
+	FExecuteAction::CreateSP(this, &SCompositePanelLayerTree::DuplicateSelectedItems),
+	FCanExecuteAction::CreateSP(this, &SCompositePanelLayerTree::CanDuplicateSelectedItems));
+	
+	CommandList->MapAction(FGenericCommands::Get().Delete,
+		FExecuteAction::CreateSP(this, &SCompositePanelLayerTree::DeleteSelectedItems),
+		FCanExecuteAction::CreateSP(this, &SCompositePanelLayerTree::CanDeleteSelectedItems));
+		
+	CommandList->MapAction(FGenericCommands::Get().Rename,
+		FUIAction(FExecuteAction::CreateSP(this, &SCompositePanelLayerTree::RenameSelectedItem),
+		FCanExecuteAction::CreateSP(this, &SCompositePanelLayerTree::CanRenameSelectedItem)));
+}
+
+void SCompositePanelLayerTree::FillCompositeActorTree()
+{
+	CompositeActorTreeItems.Empty();
+	FilteredCompositeActorTreeItems.Empty();
+	
+	if (!GEditor)
+	{
+		return;
+	}
+
+	if (const UWorld* World = GEditor->GetEditorWorldContext().World())
+	{
+		for (TActorIterator<ACompositeActor> It(World); It; ++It)
+		{
+			ACompositeActor* CompositeActor = *It;
+			if (!IsValid(CompositeActor))
+			{
+				continue;
+			}
+			
+			FCompositeActorTreeItemPtr NewActorTreeItem = MakeShared<FCompositeActorTreeItem>();
+			NewActorTreeItem->CompositeActor = CompositeActor;
+			NewActorTreeItem->bFilteredOut = Toolbar.IsValid() ? !Toolbar->ItemPassesFilters(NewActorTreeItem) : false;
+			
+			FillCompositeActorTreeItem(NewActorTreeItem);
+
+			const bool bAnyChildPassesFilter = NewActorTreeItem->Children.ContainsByPredicate([](const FCompositeActorTreeItemPtr& InChildItem)
+			{
+				return !InChildItem->bFilteredOut;
+			});
+			
+			CompositeActorTreeItems.Add(NewActorTreeItem);
+			if (!NewActorTreeItem->bFilteredOut || bAnyChildPassesFilter)
+			{
+				FilteredCompositeActorTreeItems.Add(NewActorTreeItem);
+			}
+		}
+	}
+
+	RefreshAndExpandTreeView();
+}
+
+void SCompositePanelLayerTree::FillCompositeActorTreeItem(FCompositeActorTreeItemPtr& InOutTreeItem)
+{
+	if (!InOutTreeItem.IsValid())
+	{
+		return;
+	}
+	
+	InOutTreeItem->Children.Empty();
+
+	if (ACompositeActor* CompositeActor = InOutTreeItem->CompositeActor.Get())
+	{
+		for (int32 Index = 0; Index < CompositeActor->CompositeLayers.Num(); ++Index)
+		{
+			FCompositeActorTreeItemPtr NewLayerTreeItem = MakeShared<FCompositeActorTreeItem>();
+			NewLayerTreeItem->CompositeActor = CompositeActor;
+			NewLayerTreeItem->LayerIndex = Index;
+			NewLayerTreeItem->bFilteredOut = Toolbar.IsValid() ? !Toolbar->ItemPassesFilters(NewLayerTreeItem) : false;
+			
+			InOutTreeItem->Children.Add(NewLayerTreeItem);
+		}
+	}
+}
+
+void SCompositePanelLayerTree::RefreshCompositeActorLayers(ACompositeActor* InCompositeActor, bool bRefreshTree)
+{
+	int32 TreeItemIndex = CompositeActorTreeItems.IndexOfByPredicate([InCompositeActor](const FCompositeActorTreeItemPtr& InTreeItem)
+	{
+		return InTreeItem.IsValid() && InTreeItem->CompositeActor == InCompositeActor;
+	});
+
+	if (TreeItemIndex != INDEX_NONE)
+	{
+		FillCompositeActorTreeItem(CompositeActorTreeItems[TreeItemIndex]);
+
+		if (bRefreshTree)
+		{
+			RefreshAndExpandTreeView();
+		}
+	}
+}
+
+void SCompositePanelLayerTree::SelectCompositeActorLayers(ACompositeActor* InCompositeActor, const TArray<int32>& InLayersToSelect)
+{
+	if (!TreeView.IsValid())
+	{
+		return;
+	}
+	
+	int32 TreeItemIndex = CompositeActorTreeItems.IndexOfByPredicate([InCompositeActor](const FCompositeActorTreeItemPtr& InTreeItem)
+	{
+		return InTreeItem.IsValid() && InTreeItem->CompositeActor == InCompositeActor;
+	});
+
+	if (TreeItemIndex != INDEX_NONE)
+	{
+		FCompositeActorTreeItemPtr& CompositeActorTreeItem = CompositeActorTreeItems[TreeItemIndex];
+		for (const int32 LayerIndex : InLayersToSelect)
+		{
+			if (CompositeActorTreeItem->Children.IsValidIndex(LayerIndex))
+			{
+				TreeView->SetItemSelection(CompositeActorTreeItem->Children[LayerIndex], true);	
+			}
+		}
+	}
+}
+
+void SCompositePanelLayerTree::RefreshAndExpandTreeView()
+{
+	if (TreeView.IsValid())
+	{
+		TreeView->RebuildList();
+
+		for (int32 Index = 0; Index < FilteredCompositeActorTreeItems.Num(); ++Index)
+		{
+			TreeView->SetItemExpansion(FilteredCompositeActorTreeItems[Index], true);
+		}
+	}
+}
+
+bool SCompositePanelLayerTree::CanPilotCamera() const
+{
+	TArray<FCompositeActorTreeItemPtr> SelectedItems = TreeView.IsValid() ? TreeView->GetSelectedItems() : TArray<FCompositeActorTreeItemPtr>();
+	if (!SelectedItems.Num())
+	{
+		return false;
+	}
+
+	if (!SelectedItems[0]->HasValidCompositeActor())
+	{
+		return false;
+	}
+
+	return SelectedItems[0]->CompositeActor->GetCamera().OtherActor.IsValid();
+}
+
+bool SCompositePanelLayerTree::IsPilotingCamera() const
+{
+	TArray<FCompositeActorTreeItemPtr> SelectedItems = TreeView.IsValid() ? TreeView->GetSelectedItems() : TArray<FCompositeActorTreeItemPtr>();
+	if (!SelectedItems.Num())
+	{
+		return false;
+	}
+
+	if (!SelectedItems[0]->HasValidCompositeActor())
+	{
+		return false;
+	}
+
+	AActor* CompositeCameraActor = SelectedItems[0]->CompositeActor->GetCamera().OtherActor.Get();
+	if (!CompositeCameraActor)
+	{
+		return false;
+	}
+	
+	if (ULevelEditorSubsystem* LevelEditorSubsystem = GEditor->GetEditorSubsystem<ULevelEditorSubsystem>())
+	{
+		return CompositeCameraActor == LevelEditorSubsystem->GetPilotLevelActor();
+	}
+
+	return false;
+}
+
+void SCompositePanelLayerTree::OnPilotCameraToggled(bool bPilotCamera)
+{
+	TArray<FCompositeActorTreeItemPtr> SelectedItems = TreeView.IsValid() ? TreeView->GetSelectedItems() : TArray<FCompositeActorTreeItemPtr>();
+	if (!SelectedItems.Num())
+	{
+		return;
+	}
+
+	if (!SelectedItems[0]->HasValidCompositeActor())
+	{
+		return;
+	}
+
+	AActor* CompositeCameraActor = SelectedItems[0]->CompositeActor->GetCamera().OtherActor.Get();
+	if (!CompositeCameraActor)
+	{
+		return;
+	}
+	
+	if (ULevelEditorSubsystem* LevelEditorSubsystem = GEditor->GetEditorSubsystem<ULevelEditorSubsystem>())
+	{
+		if (bPilotCamera)
+		{
+			LevelEditorSubsystem->PilotLevelActor(CompositeCameraActor);
+		}
+		else
+		{
+			LevelEditorSubsystem->EjectPilotLevelActor();
+		}
+	}
+}
+
+ECheckBoxState SCompositePanelLayerTree::GetGlobalEnabledCheckState() const
+{
+	bool bAnyEnabled = false;
+	bool bAnyDisabled = false;
+	for (const FCompositeActorTreeItemPtr& TreeItem : FilteredCompositeActorTreeItems)
+	{
+		if (!TreeItem.IsValid() || !TreeItem->HasValidCompositeActor())
+		{
+			continue;
+		}
+
+		if (TreeItem->CompositeActor->IsEnabled())
+		{
+			bAnyEnabled = true;
+		}
+		else
+		{
+			bAnyDisabled = true;
+		}
+
+		for (const FCompositeActorTreeItemPtr& ChildTreeItem : TreeItem->Children)
+		{
+			if (!ChildTreeItem.IsValid() ||
+				!ChildTreeItem->HasValidCompositeLayer() ||
+				ChildTreeItem->bFilteredOut)
+			{
+				continue;
+			}
+
+			if (UCompositeLayerBase* CompositeLayer = ChildTreeItem->GetCompositeLayer())
+			{
+				if (CompositeLayer->bIsEnabled)
+				{
+					bAnyEnabled = true;
+				}
+				else
+				{
+					bAnyDisabled = true;
+				}
+			}
+		}
+	}
+
+	if (bAnyEnabled)
+	{
+		return bAnyDisabled ? ECheckBoxState::Undetermined : ECheckBoxState::Checked;
+	}
+	else
+	{
+		return bAnyDisabled ? ECheckBoxState::Unchecked : ECheckBoxState::Undetermined;
+	}
+}
+
+void SCompositePanelLayerTree::OnGlobalEnabledCheckStateChanged(ECheckBoxState CheckBoxState)
+{
+	TSharedPtr<FScopedTransaction> GlobalSetEnabledTransaction = nullptr;
+	
+	for (const FCompositeActorTreeItemPtr& TreeItem : FilteredCompositeActorTreeItems)
+	{
+		if (!TreeItem.IsValid() || !TreeItem->HasValidCompositeActor())
+		{
+			continue;
+		}
+
+		if (!GlobalSetEnabledTransaction.IsValid())
+		{
+			GlobalSetEnabledTransaction = MakeShared<FScopedTransaction>(LOCTEXT("GlobalSetEnabledTransaction", "Set Enabled"));
+		}
+
+		TreeItem->CompositeActor->Modify();
+		TreeItem->CompositeActor->SetEnabled(CheckBoxState == ECheckBoxState::Checked);
+		
+		for (const FCompositeActorTreeItemPtr& ChildTreeItem : TreeItem->Children)
+		{
+			if (!ChildTreeItem.IsValid() ||
+				!ChildTreeItem->HasValidCompositeLayer() ||
+				ChildTreeItem->bFilteredOut)
+			{
+				continue;
+			}
+
+			if (UCompositeLayerBase* CompositeLayer = ChildTreeItem->GetCompositeLayer())
+			{
+				CompositeLayer->Modify();
+				CompositeLayer->bIsEnabled = CheckBoxState == ECheckBoxState::Checked;
+			}
+		}
+	}
+}
+
+void SCompositePanelLayerTree::OnLayerSelectionChanged(FCompositeActorTreeItemPtr InTreeItem, ESelectInfo::Type SelectInfo)
+{
+	TArray<FCompositeActorTreeItemPtr> SelectedItems = TreeView.IsValid() ? TreeView->GetSelectedItems() : TArray<FCompositeActorTreeItemPtr>();
+	
+	TArray<UObject*> SelectedLayers;
+	SelectedLayers.Reserve(SelectedItems.Num());
+
+	for (const FCompositeActorTreeItemPtr& Item : SelectedItems)
+	{
+		if (!Item->HasValidCompositeActor())
+		{
+			continue;
+		}
+
+		if (Item->LayerIndex == INDEX_NONE)
+		{
+			SelectedLayers.Add(Item->CompositeActor.Get());
+		}
+		else if (Item->HasValidCompositeLayer())
+		{
+			if (UCompositeLayerBase* CompositeLayer = Item->GetCompositeLayer())
+			{
+				SelectedLayers.Add(CompositeLayer);
+			}
+		}
+	}
+	
+	OnSelectionChanged.ExecuteIfBound(SelectedLayers);
+}
+
+void SCompositePanelLayerTree::OnFilterChanged()
+{
+	FilteredCompositeActorTreeItems.Empty();
+	
+	for (FCompositeActorTreeItemPtr& TreeItem : CompositeActorTreeItems)
+	{
+		TreeItem->bFilteredOut = Toolbar.IsValid() ? !Toolbar->ItemPassesFilters(TreeItem) : false;
+
+		bool bChildLayerPassesFilter = false;
+		for (FCompositeActorTreeItemPtr& ChildTreeItem : TreeItem->Children)
+		{
+			ChildTreeItem->bFilteredOut = Toolbar.IsValid() ? !Toolbar->ItemPassesFilters(ChildTreeItem) : false;
+			if (!ChildTreeItem->bFilteredOut)
+			{
+				bChildLayerPassesFilter = true;
+			}
+		}
+
+		if (!TreeItem->bFilteredOut || bChildLayerPassesFilter)
+		{
+			FilteredCompositeActorTreeItems.Add(TreeItem);
+		}
+	}
+	
+	RefreshAndExpandTreeView();
+}
+
+FText SCompositePanelLayerTree::GetFilterStatusText() const
+{
+	const int32 NumActors = CompositeActorTreeItems.Num();
+	const int32 NumFiltered = FilteredCompositeActorTreeItems.Num();
+
+	TArray<FCompositeActorTreeItemPtr> SelectedItems = TreeView.IsValid() ? TreeView->GetSelectedItems() : TArray<FCompositeActorTreeItemPtr>();
+	const int32 NumSelected = SelectedItems.Num();
+	
+	const FText ActorLabel = NumActors > 1 ? LOCTEXT("ActorPlural", "Actors") : LOCTEXT("ActorSingular", "Actor");
+	if (NumActors == NumFiltered)
+	{
+		if (NumSelected > 0)
+		{
+			return FText::Format(LOCTEXT("NumActorsAndSelectedTextFormat", "{0} {1}, {2} Selected"), FText::AsNumber(NumActors), ActorLabel, FText::AsNumber(NumSelected));
+		}
+		else
+		{
+			return FText::Format(LOCTEXT("NumActorsTextFormat", "{0} {1}"), FText::AsNumber(NumActors), ActorLabel);
+		}
+	}
+	else if (NumFiltered > 0)
+	{
+		if (NumSelected > 0)
+		{
+			return FText::Format(LOCTEXT("NumActorsWithFilteredAndSelectedTextFormat", "Showing {0} of {1} {2}, {3} Selected"),
+				FText::AsNumber(NumFiltered),
+				FText::AsNumber(NumActors),
+				ActorLabel,
+				FText::AsNumber(NumSelected));
+		}
+		else
+		{
+			return FText::Format(LOCTEXT("NumActorsWithFilteredTextFormat", "Showing {0} of {1} {2}"),
+				FText::AsNumber(NumFiltered),
+				FText::AsNumber(NumActors),
+				ActorLabel);
+		}
+	}
+	else
+	{
+		if (NumActors > 0)
+		{
+			return FText::Format(LOCTEXT("NoMatchingActorsTextFormat", "No matching actors or layers ({0} {1})"), FText::AsNumber(NumActors), ActorLabel);
+		}
+		else
+		{
+			return LOCTEXT("NoActorsFoundLabel", "No composite actors found in level");
+		}
+	}
+}
+
+FSlateColor SCompositePanelLayerTree::GetFilterStatusTextColor() const
+{
+	const int32 NumActors = CompositeActorTreeItems.Num();
+	const int32 NumFiltered = FilteredCompositeActorTreeItems.Num();
+	
+	if (NumActors == 0 || NumFiltered == 0)
+	{
+		return FAppStyle::Get().GetSlateColor("Colors.AccentRed");
+	}
+	
+	return FSlateColor::UseForeground();
+}
+
+TSharedPtr<SWidget> SCompositePanelLayerTree::CreateTreeContextMenu()
+{
+	if (!TreeView.IsValid())
+	{
+		return SNullWidget::NullWidget;
+	}
+	
+	const int32 NumItems = TreeView->GetNumItemsSelected();
+	if (NumItems >= 1)
+	{
+		constexpr bool bShouldCloseWindowAfterMenuSelection = true;
+		FMenuBuilder MenuBuilder(bShouldCloseWindowAfterMenuSelection, CommandList);
+
+		MenuBuilder.AddMenuEntry(FGenericCommands::Get().Copy);
+		MenuBuilder.AddMenuEntry(FGenericCommands::Get().Cut);
+		MenuBuilder.AddMenuEntry(FGenericCommands::Get().Paste);
+		MenuBuilder.AddMenuEntry(FGenericCommands::Get().Duplicate);
+		MenuBuilder.AddMenuEntry(FGenericCommands::Get().Delete);
+		MenuBuilder.AddMenuEntry(FGenericCommands::Get().Rename);
+		
+		return MenuBuilder.MakeWidget();
+	}
+
+	return SNullWidget::NullWidget;
+}
+
+void SCompositePanelLayerTree::OnLayerAdded(const UClass* InLayerClass)
+{
+	TArray<FCompositeActorTreeItemPtr> SelectedItems = TreeView.IsValid() ? TreeView->GetSelectedItems() : TArray<FCompositeActorTreeItemPtr>();
+	if (!SelectedItems.Num() || !SelectedItems[0].IsValid())
+	{
+		return;
+	}
+
+	ACompositeActor* CompositeActor = SelectedItems[0]->CompositeActor.Get();
+	if (!CompositeActor)
+	{
+		return;
+	}
+
+	const int32 NewLayerIndex = SelectedItems[0]->LayerIndex + 1;
+	FScopedTransaction LayerAddedTransaction(LOCTEXT("LayerAddedTransaction", "Add Layer"));
+	CompositeActor->Modify();
+
+	UCompositeLayerBase* NewLayer = NewObject<UCompositeLayerBase>(CompositeActor, InLayerClass);
+	NewLayer->Name = NewLayer->GetName();
+	
+	CompositeActor->CompositeLayers.Insert(NewLayer, NewLayerIndex);
+
+	RefreshCompositeActorLayers(CompositeActor);
+}
+
+FReply SCompositePanelLayerTree::OnLayerSoloToggled(const FCompositeActorTreeItemPtr& InLayerTreeItem)
+{
+	if (!InLayerTreeItem.IsValid() ||
+		!InLayerTreeItem->HasValidCompositeLayer())
+	{
+		return FReply::Unhandled();
+	}
+
+	ACompositeActor* CompositeActor = InLayerTreeItem->CompositeActor.Get();
+	UCompositeLayerBase* CompositeLayer = InLayerTreeItem->GetCompositeLayer();
+
+	if (!CompositeLayer)
+	{
+		return FReply::Unhandled();
+	}
+	
+	FScopedTransaction SoloLayerTransaction(LOCTEXT("SoloLayerTransaction", "Solo Layer"));
+	CompositeActor->Modify();
+	CompositeLayer->Modify();
+
+	CompositeLayer->bIsSolo = !CompositeLayer->bIsSolo;
+
+	if (CompositeLayer->bIsSolo)
+	{
+		// Unsolo all other layers if the specified layer has been made solo
+		for (TObjectPtr<UCompositeLayerBase>& OtherLayer : CompositeActor->CompositeLayers)
+		{
+			if (OtherLayer != CompositeLayer)
+			{
+				OtherLayer->bIsSolo = false;
+			}
+		}
+	}
+	
+	return FReply::Handled();
+}
+
+void SCompositePanelLayerTree::OnLayerMoved(const FCompositeActorTreeItemPtr& InTreeItem, int32 InDestIndex)
+{
+	if (!InTreeItem.IsValid() ||
+		!InTreeItem->HasValidCompositeActor() ||
+		!InTreeItem->HasValidCompositeLayer())
+	{
+		return;
+	}
+	
+	ACompositeActor* CompositeActor = InTreeItem->CompositeActor.Get();
+	UCompositeLayerBase* CompositeLayer = InTreeItem->GetCompositeLayer();
+	
+	if (!CompositeLayer || !CompositeActor->CompositeLayers.IsValidIndex(InDestIndex))
+	{
+		return;
+	}
+	
+	FScopedTransaction MoveLayerTransaction(LOCTEXT("MoveLayerTransaction", "Move Layer"));
+	CompositeActor->Modify();
+	
+	CompositeActor->CompositeLayers.RemoveAt(InTreeItem->LayerIndex);
+	CompositeActor->CompositeLayers.Insert(CompositeLayer, InDestIndex);
+	
+	RefreshCompositeActorLayers(CompositeActor);
+}
+
+void SCompositePanelLayerTree::OnActorAddedToLevel(AActor* Actor)
+{
+	if (Actor->IsA<ACompositeActor>())
+	{
+		FillCompositeActorTree();
+	}
+}
+
+void SCompositePanelLayerTree::OnLevelActorListChanged()
+{
+	FillCompositeActorTree();
+}
+
+void SCompositePanelLayerTree::OnActorRemovedFromLevel(AActor* Actor)
+{
+	if (Actor->IsA<ACompositeActor>())
+	{
+		FillCompositeActorTree();
+	}
+}
+
+void SCompositePanelLayerTree::OnMapChange(uint32 MapFlags)
+{
+	FillCompositeActorTree();
+}
+
+void SCompositePanelLayerTree::OnNewCurrentLevel()
+{
+	FillCompositeActorTree();
+}
+
+void SCompositePanelLayerTree::OnLevelAdded(ULevel* InLevel, UWorld* InWorld)
+{
+	FillCompositeActorTree();
+}
+
+void SCompositePanelLayerTree::OnLevelRemoved(ULevel* InLevel, UWorld* InWorld)
+{
+	FillCompositeActorTree();
+}
+
+void SCompositePanelLayerTree::OnObjectReplaced(const TMap<UObject*, UObject*>& Tuples)
+{
+	bool bContainsCompositeActor = false;
+	for (const FCompositeActorTreeItemPtr& TreeItem : CompositeActorTreeItems)
+	{
+		if (Tuples.Contains(TreeItem->CompositeActor.GetEvenIfUnreachable()))
+		{
+			bContainsCompositeActor = true;
+		}
+	}
+
+	if (bContainsCompositeActor)
+	{
+		FillCompositeActorTree();
+	}
+}
+
+void SCompositePanelLayerTree::OnObjectPropertyChanged(UObject* Object, FPropertyChangedEvent& PropertyChangedEvent)
+{
+	if (Object->IsA<ACompositeActor>() && PropertyChangedEvent.GetPropertyName() == GET_MEMBER_NAME_CHECKED(ACompositeActor, CompositeLayers))
+	{
+		RefreshCompositeActorLayers(Cast<ACompositeActor>(Object));
+	}
+}
+
+void SCompositePanelLayerTree::CopySelectedItems()
+{
+	TArray<UObject*> ObjectsToCopy;
+
+	TArray<FCompositeActorTreeItemPtr> SelectedItems = TreeView.IsValid() ? TreeView->GetSelectedItems() : TArray<FCompositeActorTreeItemPtr>();
+	for (const FCompositeActorTreeItemPtr& SelectedItem : SelectedItems)
+	{
+		if (SelectedItem.IsValid() && SelectedItem->HasValidCompositeLayer())
+		{
+			if (UCompositeLayerBase* CompositeLayer = SelectedItem->GetCompositeLayer())
+			{
+				ObjectsToCopy.Add(CompositeLayer);
+			}
+		}
+	}
+
+	if (ObjectsToCopy.IsEmpty())
+	{
+		return;
+	}
+	
+	UnMarkAllObjects(EObjectMark(OBJECTMARK_TagExp | OBJECTMARK_TagImp));
+
+	FStringOutputDevice Archive;
+	const FExportObjectInnerContext Context;
+
+	for (UObject* Object : ObjectsToCopy)
+	{
+		UExporter::ExportToOutputDevice(&Context, Object, nullptr, Archive, TEXT("copy"), 0, PPF_ExportsNotFullyQualified | PPF_Copy | PPF_Delimited, false, nullptr);
+	}
+
+	FPlatformApplicationMisc::ClipboardCopy(*Archive);
+}
+
+bool SCompositePanelLayerTree::CanCopySelectedItems()
+{
+	TArray<FCompositeActorTreeItemPtr> SelectedItems = TreeView.IsValid() ? TreeView->GetSelectedItems() : TArray<FCompositeActorTreeItemPtr>();
+	if (!SelectedItems.Num())
+	{
+		return false;
+	}
+
+	const bool bContainsLayer = SelectedItems.ContainsByPredicate([](const FCompositeActorTreeItemPtr& InTreeItem)
+	{
+		return InTreeItem.IsValid() && InTreeItem->HasValidCompositeLayer() && InTreeItem->GetCompositeLayer();
+	});
+
+	return bContainsLayer;
+}
+
+void SCompositePanelLayerTree::CutSelectedItems()
+{
+	CopySelectedItems();
+	DeleteSelectedItems();
+}
+
+bool SCompositePanelLayerTree::CanCutSelectedItems()
+{
+	return CanCopySelectedItems() && CanDeleteSelectedItems();
+}
+
+class FCompositeLayerObjectTextFactory : public FCustomizableTextObjectFactory
+{
+public:
+	FCompositeLayerObjectTextFactory() : FCustomizableTextObjectFactory(GWarn) { }
+
+	virtual bool CanCreateClass(UClass* ObjectClass, bool& bOmitSubObjs) const override
+	{
+		return ObjectClass->IsChildOf(UCompositeLayerBase::StaticClass());
+	}
+
+	virtual void ProcessConstructedObject(UObject* NewObject) override
+	{
+		check(NewObject);
+
+		if (UCompositeLayerBase* CompositeLayer = Cast<UCompositeLayerBase>(NewObject))
+		{
+			CompositeLayers.Add(CompositeLayer);
+		}
+	}
+
+public:
+	TArray<UCompositeLayerBase*> CompositeLayers;
+};
+
+void SCompositePanelLayerTree::PasteSelectedItems()
+{
+	TArray<FCompositeActorTreeItemPtr> SelectedItems = TreeView.IsValid() ? TreeView->GetSelectedItems() : TArray<FCompositeActorTreeItemPtr>();
+	if (SelectedItems.IsEmpty())
+	{
+		return;
+	}
+
+	// A list of composite actors to paste the clipboard contents in, and the layer index to paste at
+	TMap<ACompositeActor*, int32> CompositeActorsToPasteTo;
+	for (const FCompositeActorTreeItemPtr& SelectedItem : SelectedItems)
+	{
+		if (SelectedItem.IsValid() && SelectedItem->HasValidCompositeActor())
+		{
+			if (CompositeActorsToPasteTo.Contains(SelectedItem->CompositeActor.Get()))
+			{
+				CompositeActorsToPasteTo[SelectedItem->CompositeActor.Get()] = FMath::Max(SelectedItem->LayerIndex + 1, CompositeActorsToPasteTo[SelectedItem->CompositeActor.Get()]);
+			}
+			else
+			{
+				CompositeActorsToPasteTo.Add(SelectedItem->CompositeActor.Get(), SelectedItem->LayerIndex + 1);
+			}
+		}
+	}
+
+	if (CompositeActorsToPasteTo.IsEmpty())
+	{
+		return;
+	}
+
+	FString ClipboardContent;
+	FPlatformApplicationMisc::ClipboardPaste(ClipboardContent);
+	
+	FCompositeLayerObjectTextFactory Factory;
+	Factory.ProcessBuffer(GetTransientPackage(), RF_Transactional | RF_Transient, ClipboardContent);
+
+	if (Factory.CompositeLayers.IsEmpty())
+	{
+		return;
+	}
+
+	TMap<ACompositeActor*, TArray<int32>> LayersToSelect;
+	FScopedTransaction PasteLayersTransaction(LOCTEXT("PasteLayersTransaction", "Paste Layers"));
+	for (const TPair<ACompositeActor*, int32>& Pair : CompositeActorsToPasteTo)
+	{
+		ACompositeActor* CompositeActor = Pair.Key;
+		int32 IndexToPasteAt = Pair.Value;
+
+		LayersToSelect.Add(CompositeActor, TArray<int32>());
+		
+		CompositeActor->Modify();
+		for (UCompositeLayerBase* LayerToPaste : Factory.CompositeLayers)
+		{
+			UCompositeLayerBase* NewLayer = DuplicateObject(LayerToPaste, CompositeActor);
+			CompositeActor->CompositeLayers.Insert(NewLayer, IndexToPasteAt);
+			LayersToSelect[CompositeActor].Add(IndexToPasteAt);
+			++IndexToPasteAt;
+		}
+
+		constexpr bool bRefreshTree = false;
+		RefreshCompositeActorLayers(CompositeActor, bRefreshTree);
+	}
+
+	RefreshAndExpandTreeView();
+
+	for (const TPair<ACompositeActor*, TArray<int32>>& Pair : LayersToSelect)
+	{
+		SelectCompositeActorLayers(Pair.Key, Pair.Value);
+	}
+}
+
+bool SCompositePanelLayerTree::CanPasteSelectedItems()
+{
+	TArray<FCompositeActorTreeItemPtr> SelectedItems = TreeView.IsValid() ? TreeView->GetSelectedItems() : TArray<FCompositeActorTreeItemPtr>();
+	if (SelectedItems.IsEmpty())
+	{
+		return false;
+	}
+
+	const bool bHasValidCompositeActor = SelectedItems.ContainsByPredicate([](const FCompositeActorTreeItemPtr& InTreeItem)
+	{
+		return InTreeItem.IsValid() && InTreeItem->HasValidCompositeActor();
+	});
+
+	if (!bHasValidCompositeActor)
+	{
+		return false;
+	}
+	
+	// Can't paste unless the clipboard has a string in it
+	FString ClipboardContent;
+	FPlatformApplicationMisc::ClipboardPaste(ClipboardContent);
+	if (!ClipboardContent.IsEmpty())
+	{
+		FCompositeLayerObjectTextFactory Factory;
+		Factory.ProcessBuffer(GetTransientPackage(), RF_Transactional | RF_Transient, ClipboardContent);
+		return !Factory.CompositeLayers.IsEmpty();
+	}
+
+	return false;
+}
+
+void SCompositePanelLayerTree::DuplicateSelectedItems()
+{
+	CopySelectedItems();
+	PasteSelectedItems();
+}
+
+bool SCompositePanelLayerTree::CanDuplicateSelectedItems()
+{
+	return CanCopySelectedItems();
+}
+
+void SCompositePanelLayerTree::DeleteSelectedItems()
+{
+	TArray<FCompositeActorTreeItemPtr> SelectedItems = TreeView.IsValid() ? TreeView->GetSelectedItems() : TArray<FCompositeActorTreeItemPtr>();
+
+	TMap<ACompositeActor*, TArray<int32>> LayersToDelete;
+	for (const FCompositeActorTreeItemPtr& Item : SelectedItems)
+	{
+		if (!Item.IsValid() ||
+			!Item->HasValidCompositeLayer())
+		{
+			continue;
+		}
+
+		if (LayersToDelete.Contains(Item->CompositeActor.Get()))
+		{
+			LayersToDelete[Item->CompositeActor.Get()].Add(Item->LayerIndex);
+		}
+		else
+		{
+			LayersToDelete.Add(Item->CompositeActor.Get(), TArray<int32> { Item->LayerIndex });
+		}
+	}
+
+	if (LayersToDelete.Num())
+	{
+		FScopedTransaction DeleteLayerTransaction(LOCTEXT("DeleteLayerTransaction", "Delete Layer"));
+		for (TPair<ACompositeActor*, TArray<int32>>& CompositeActorLayerList : LayersToDelete)
+		{
+			ACompositeActor* CompositeActor = CompositeActorLayerList.Key;
+			CompositeActor->Modify();
+			
+			// Delete the layers in reverse order so that indices don't get messed up as we delete
+			TArray<int32>& Layers = CompositeActorLayerList.Value;
+			Layers.Sort();
+			
+			for (int32 Index = Layers.Num() - 1; Index >= 0; --Index)
+			{
+				CompositeActor->CompositeLayers.RemoveAt(Layers[Index]);
+			}
+
+			int32 TreeItemIndex = CompositeActorTreeItems.IndexOfByPredicate([CompositeActor](const FCompositeActorTreeItemPtr& InTreeItem)
+			{
+				return InTreeItem->CompositeActor.Get() == CompositeActor;
+			});
+
+			if (TreeItemIndex != INDEX_NONE)
+			{
+				FillCompositeActorTreeItem(CompositeActorTreeItems[TreeItemIndex]);
+			}
+		}
+
+		RefreshAndExpandTreeView();
+	}
+}
+
+bool SCompositePanelLayerTree::CanDeleteSelectedItems()
+{
+	TArray<FCompositeActorTreeItemPtr> SelectedItems = TreeView.IsValid() ? TreeView->GetSelectedItems() : TArray<FCompositeActorTreeItemPtr>();
+	const bool bContainsLayerInSelection = SelectedItems.ContainsByPredicate([](const FCompositeActorTreeItemPtr& InTreeItem)
+	{
+		return InTreeItem.IsValid() && InTreeItem->HasValidCompositeLayer();
+	});
+	
+	return bContainsLayerInSelection;
+}
+
+void SCompositePanelLayerTree::RenameSelectedItem()
+{
+	TArray<FCompositeActorTreeItemPtr> SelectedItems = TreeView.IsValid() ? TreeView->GetSelectedItems() : TArray<FCompositeActorTreeItemPtr>();
+	if (SelectedItems.Num() != 1 || !SelectedItems[0]->HasValidCompositeLayer() || !SelectedItems[0]->GetCompositeLayer())
+	{
+		return;
+	}
+	
+	TSharedPtr<ITableRow> RowWidget = TreeView->WidgetFromItem(SelectedItems[0]);
+	if (TSharedPtr<SCompositePanelLayerTreeItemRow> TreeItemRowWidget = StaticCastSharedPtr<SCompositePanelLayerTreeItemRow>(RowWidget))
+	{
+		TreeItemRowWidget->RequestRename();
+	}
+}
+
+bool SCompositePanelLayerTree::CanRenameSelectedItem()
+{
+	TArray<FCompositeActorTreeItemPtr> SelectedItems = TreeView.IsValid() ? TreeView->GetSelectedItems() : TArray<FCompositeActorTreeItemPtr>();
+	if (SelectedItems.Num() == 1)
+	{
+		return SelectedItems[0].IsValid() && SelectedItems[0]->HasValidCompositeLayer() && SelectedItems[0]->GetCompositeLayer();
+	}
+	
+	return false;
+}
+
+#undef LOCTEXT_NAMESPACE
diff --git a/Engine/Plugins/Experimental/Compositing/Composite/Source/CompositeEditor/Private/UI/SCompositePanelLayerTree.h b/Engine/Plugins/Experimental/Compositing/Composite/Source/CompositeEditor/Private/UI/SCompositePanelLayerTree.h
new file mode 100644
--- /dev/null
+++ b/Engine/Plugins/Experimental/Compositing/Composite/Source/CompositeEditor/Private/UI/SCompositePanelLayerTree.h
@@ -0,0 +1,172 @@
+﻿// Copyright Epic Games, Inc. All Rights Reserved.
+
+#pragma once
+
+#include "EditorUndoClient.h"
+#include "Layers/CompositeLayerBase.h"
+#include "Widgets/SCompoundWidget.h"
+
+class FUICommandList;
+class SCompositePanelLayerTreeToolbar;
+class ACompositeActor;
+template<typename T> class STreeView;
+
+/**
+ * A hierarchical view that displays all compositing actors and their rendering layers in the level
+ */
+class SCompositePanelLayerTree : public SCompoundWidget, public FEditorUndoClient
+{
+public:
+	DECLARE_DELEGATE_OneParam(FOnSelectionChanged, const TArray<UObject*>&);
+	
+private:
+	/** Item that stores a composite actor and layer index to display in the layer tree view */
+	struct FCompositeActorTreeItem
+	{
+		/** Composite actor whose layer is being displayed */
+		TWeakObjectPtr<ACompositeActor> CompositeActor = nullptr;
+
+		/** Index of the layer in the composite actor's CompositeLayers list. If INDEX_NONE, then this tree item is the composite actor itself */
+		int32 LayerIndex = INDEX_NONE;
+
+		/** Child tree items of this tree item */
+		TArray<TSharedPtr<FCompositeActorTreeItem>> Children;
+
+		/** Indicates that this item is being filtered out by the active filter */
+		bool bFilteredOut = false;
+
+		/** Gets whether this tree item references a valid composite actor */
+		bool HasValidCompositeActor() const;
+
+		/** Gets whether this tree item references a valid composite layer */
+		bool HasValidCompositeLayer() const;
+
+		/** Gets the composite layer this tree item references */
+		UCompositeLayerBase* GetCompositeLayer() const;
+	};
+
+	using FCompositeActorTreeItemPtr = TSharedPtr<FCompositeActorTreeItem>;
+	
+public:
+	SLATE_BEGIN_ARGS(SCompositePanelLayerTree) { }
+		SLATE_EVENT(FOnSelectionChanged, OnSelectionChanged)
+	SLATE_END_ARGS()
+
+	void Construct(const FArguments& InArgs);
+
+	virtual ~SCompositePanelLayerTree() override;
+
+	// SWidget interface
+	virtual FReply OnKeyDown(const FGeometry& MyGeometry, const FKeyEvent& InKeyEvent) override;
+	// End of SWidget interface
+	
+	// FEditorUndoClient Interface
+	virtual void PostUndo(bool bSuccess) override;
+	virtual void PostRedo(bool bSuccess) override;
+	// End of FEditorUndoClient
+	
+private:	
+	/** Binds commands and callbacks to the command list */
+	void BindCommands();
+	
+	/** Fills the tree view with all composite actors found on the current level */
+	void FillCompositeActorTree();
+
+	/** Fills the children of the composite actor tree item with tree items for each of the composite actor's layers */
+	void FillCompositeActorTreeItem(FCompositeActorTreeItemPtr& InOutTreeItem);
+
+	/** Refreshes the child tree elements for the specified composite actor */
+	void RefreshCompositeActorLayers(ACompositeActor* InCompositeActor, bool bRefreshTree = true);
+
+	/** Selects the specified layers of a composite actor in the tree view */
+	void SelectCompositeActorLayers(ACompositeActor* InCompositeActor, const TArray<int32>& InLayersToSelect);
+	
+	/** Refreshes the tree view and expands all tree view items */
+	void RefreshAndExpandTreeView();
+
+	/** Gets whether the currently selected composite actor has a camera that can be piloted */
+	bool CanPilotCamera() const;
+
+	/** Gets whether camera actor of the currently selected composite actor is being piloted */
+	bool IsPilotingCamera() const;
+
+	/** Sets the piloted actor to the currently selected composite actor camera, or clears it if it is already being piloted */
+	void OnPilotCameraToggled(bool bPilotCamera);
+
+	/** Gets the check state for the global enabled checkbox */
+	ECheckBoxState GetGlobalEnabledCheckState() const;
+
+	/** Sets the Enabled flag for all currently filtered items in the tree view */
+	void OnGlobalEnabledCheckStateChanged(ECheckBoxState CheckBoxState);
+
+	/** Callback that is raised when the selected items in the tree view change */
+	void OnLayerSelectionChanged(FCompositeActorTreeItemPtr InTreeItem, ESelectInfo::Type SelectInfo);
+
+	/** Callback that is raised when the tree view filter is changed, applies the new filter to the composite actor tree items */
+	void OnFilterChanged();
+
+	/** Gets the status text to display at the bottom of the tree view that shows the number of actors filtered and selected */
+	FText GetFilterStatusText() const;
+
+	/** Gets the color of the filter status text */
+	FSlateColor GetFilterStatusTextColor() const;
+
+	/** Creates the right click context menu for the tree view */
+	TSharedPtr<SWidget> CreateTreeContextMenu();
+
+	/** Raised when a new layer is being added to the selected composite actor */
+	void OnLayerAdded(const UClass* InLayerClass);
+
+	/** Raised when the Solo button has been toggled on the specified layer tree item */
+	FReply OnLayerSoloToggled(const FCompositeActorTreeItemPtr& InLayerTreeItem);
+
+	/** Raised when a layer has been moved via a drag drop operation to the specified destination index */
+	void OnLayerMoved(const FCompositeActorTreeItemPtr& InTreeItem, int32 InDestIndex);
+
+	/** Callbacks for any editor events that may necessitate updating the composite actors being displayed in the tree view */
+	void OnActorAddedToLevel(AActor* Actor);
+	void OnLevelActorListChanged();
+	void OnActorRemovedFromLevel(AActor* Actor);
+	void OnMapChange(uint32 MapFlags);
+	void OnNewCurrentLevel();
+	void OnLevelAdded(ULevel* InLevel, UWorld* InWorld);
+	void OnLevelRemoved(ULevel* InLevel, UWorld* InWorld);
+	void OnObjectReplaced(const TMap<UObject*, UObject*>& Tuples);
+	void OnObjectPropertyChanged(UObject* Object, FPropertyChangedEvent& PropertyChangedEvent);
+
+	/** Generic command callbacks to handle copy, cut, paste, delete, and rename */
+	void CopySelectedItems();
+	bool CanCopySelectedItems();
+	void CutSelectedItems();
+	bool CanCutSelectedItems();
+	void PasteSelectedItems();
+	bool CanPasteSelectedItems();
+	void DuplicateSelectedItems();
+	bool CanDuplicateSelectedItems();
+	void DeleteSelectedItems();
+	bool CanDeleteSelectedItems();
+	void RenameSelectedItem();
+	bool CanRenameSelectedItem();
+	
+private:
+	/** Tree view widget that displays the level's composite actors and their layers */
+	TSharedPtr<STreeView<FCompositeActorTreeItemPtr>> TreeView;
+
+	/** A list of all tree items that can be displayed in the tree view */
+	TArray<FCompositeActorTreeItemPtr> CompositeActorTreeItems;
+
+	/** The actual list of tree items to display when a filter has been applied */
+	TArray<FCompositeActorTreeItemPtr> FilteredCompositeActorTreeItems;
+	
+	/** Toolbar being displayed above the tree view */
+	TSharedPtr<SCompositePanelLayerTreeToolbar> Toolbar;
+
+	/** The command list for commands related to the tree view */
+	TSharedPtr<FUICommandList> CommandList;
+
+	/** Raised when the selected tree items have changed */
+	FOnSelectionChanged OnSelectionChanged;
+	
+	friend class SCompositePanelLayerTreeItemRow;
+	friend class SCompositePanelLayerTreeToolbar;
+};
diff --git a/Engine/Plugins/Experimental/Compositing/Composite/Source/CompositeEditor/Public/CompositeEditorModule.h b/Engine/Plugins/Experimental/Compositing/Composite/Source/CompositeEditor/Public/CompositeEditorModule.h
--- a/Engine/Plugins/Experimental/Compositing/Composite/Source/CompositeEditor/Public/CompositeEditorModule.h
+++ b/Engine/Plugins/Experimental/Compositing/Composite/Source/CompositeEditor/Public/CompositeEditorModule.h
@@ -23,4 +23,11 @@ public:
 
 	/** Called when an Actor has been added to a level. */
 	void OnLevelActorAdded(AActor* InActor);
+
+private:
+	void RegisterTabSpawners();
+	void UnregisterTabSpawners();
+
+private:
+	FDelegateHandle OnTabManagerChangedDelegateHandle;
 };
