Composite: Added comment description for plate layer.

diff --git a/Engine/Plugins/Experimental/Compositing/Composite/Source/Composite/Public/Layers/CompositeLayerPlate.h b/Engine/Plugins/Experimental/Compositing/Composite/Source/Composite/Public/Layers/CompositeLayerPlate.h
--- a/Engine/Plugins/Experimental/Compositing/Composite/Source/Composite/Public/Layers/CompositeLayerPlate.h
+++ b/Engine/Plugins/Experimental/Compositing/Composite/Source/Composite/Public/Layers/CompositeLayerPlate.h
@@ -16,6 +16,12 @@ class UTexture;
 class FObjectPreSaveContext;
 class UTextureRenderTarget2D;
 
+/**
+* Convenience layer for media texture integration, with additional processing options.
+* 
+* Note that when "Holdout Composite" is active, the media texture is projected onto composite meshes & rendered separately as a custom render pass.
+* The result is dilated and combined with other layers such as the main render (with composite meshes rendered as holdouts).
+*/
 UCLASS(MinimalAPI, BlueprintType, Blueprintable, EditInlineNew, CollapseCategories, meta = (DisplayName = "Plate"))
 class UCompositeLayerPlate : public UCompositeLayerBase, public FTickableGameObject
 {
