Composite: Fix render target crash on resolution change undo.

* Replace incorrect use of ConditionalBeginDestroy by MarkAsGarbage.
* Unify calls into a single function.
* Remove previous assignments.

diff --git a/Engine/Plugins/Experimental/Compositing/Composite/Source/Composite/Private/CompositeActor.cpp b/Engine/Plugins/Experimental/Compositing/Composite/Source/Composite/Private/CompositeActor.cpp
--- a/Engine/Plugins/Experimental/Compositing/Composite/Source/Composite/Private/CompositeActor.cpp
+++ b/Engine/Plugins/Experimental/Compositing/Composite/Source/Composite/Private/CompositeActor.cpp
@@ -31,6 +31,7 @@
 #if WITH_EDITOR
 #include "LevelEditor.h"
 #include "Modules/ModuleManager.h"
+#include "Misc/TransactionObjectEvent.h"
 
 /** The editor object. */
 extern UNREALED_API class UEditorEngine* GEditor;
@@ -456,6 +457,28 @@ void ACompositeActor::PostEditChangeChainProperty(struct FPropertyChangedChainEv
 	}
 }
 
+void ACompositeActor::PostTransacted(const FTransactionObjectEvent& TransactionEvent)
+{
+	Super::PostTransacted(TransactionEvent);
+
+	if (TransactionEvent.GetEventType() == ETransactionObjectEventType::UndoRedo)
+	{
+		const FName RenderResolutionName = GET_MEMBER_NAME_CHECKED(ThisClass, RenderResolution);
+		if (TransactionEvent.GetChangedProperties().Contains(RenderResolutionName))
+		{
+			/**
+			* When a resolution change is undone, the previous render targets are deserialized & replaced.
+			* This leaves previous assignments invalid so we make sure to return them to the pool.
+			*/
+			TInlineComponentArray<USceneCaptureComponent2D*> SceneCaptureComponents(this);
+			for (USceneCaptureComponent2D* CaptureComponent : SceneCaptureComponents)
+			{
+				FCompositeRenderTargetPool::Get().ReleaseAssigneeTargets(CaptureComponent);
+			}
+		}
+	}
+}
+
 void ACompositeActor::SetIsTemporarilyHiddenInEditor(bool bIsHidden)
 {
 	Super::SetIsTemporarilyHiddenInEditor(bIsHidden);
diff --git a/Engine/Plugins/Experimental/Compositing/Composite/Source/Composite/Private/CompositeRenderTargetPool.cpp b/Engine/Plugins/Experimental/Compositing/Composite/Source/Composite/Private/CompositeRenderTargetPool.cpp
--- a/Engine/Plugins/Experimental/Compositing/Composite/Source/Composite/Private/CompositeRenderTargetPool.cpp
+++ b/Engine/Plugins/Experimental/Compositing/Composite/Source/Composite/Private/CompositeRenderTargetPool.cpp
@@ -16,6 +16,27 @@ static TAutoConsoleVariable<int32> CVarCompositePoolFramesUntilFlush(
 
 const FIntPoint FCompositeRenderTargetPool::DefaultSize = FIntPoint(1920, 1080);
 
+namespace UE
+{
+	namespace Composite
+	{
+		namespace Private
+		{
+			/** Convenience function to release & mark as garbage a render target. */
+			void Release(TObjectPtr<UTextureRenderTarget2D>& InTarget)
+			{
+				if (IsValid(InTarget))
+				{
+					InTarget->ReleaseResource();
+					InTarget->RemoveFromRoot();
+					InTarget->MarkAsGarbage();
+				}
+				InTarget = nullptr;
+			}
+		}
+	}
+}
+
 FCompositeRenderTargetPool& FCompositeRenderTargetPool::Get()
 {
 	static TUniquePtr<FCompositeRenderTargetPool> PoolInstance = MakeUnique<FCompositeRenderTargetPool>();
@@ -116,6 +137,26 @@ void FCompositeRenderTargetPool::ReleaseTarget(TObjectPtr<UTextureRenderTarget2D
 	}
 }
 
+void FCompositeRenderTargetPool::ReleaseAssigneeTargets(UObject* InAssignee)
+{
+	if (!IsValid(InAssignee))
+	{
+		return;
+	}
+
+	for (auto TargetIt = AssignedTargets.CreateIterator(); TargetIt; ++TargetIt)
+	{
+		if (InAssignee == TargetIt->Value)
+		{
+			const TObjectPtr<UTextureRenderTarget2D>& RenderTarget = TargetIt->Key;
+			const FRenderTargetDesc Desc{ FIntPoint(RenderTarget->SizeX, RenderTarget->SizeY), RenderTarget->RenderTargetFormat };
+			PooledTargets.FindOrAdd(Desc).Push(RenderTarget);
+
+			TargetIt.RemoveCurrent();
+		}
+	}
+}
+
 void FCompositeRenderTargetPool::EmptyPool(bool bInForceCollectGarbage)
 {
 	UE_LOG(LogComposite, VeryVerbose, TEXT("Empty pool."));
@@ -124,10 +165,7 @@ void FCompositeRenderTargetPool::EmptyPool(bool bInForceCollectGarbage)
 	{
 		for (auto TargetIt = DescIt->Value.CreateIterator(); TargetIt; ++TargetIt)
 		{
-			TObjectPtr<UTextureRenderTarget2D>& Target = TargetIt->Texture;
-			Target->ReleaseResource();
-			Target->ConditionalBeginDestroy();
-			Target = nullptr;
+			UE::Composite::Private::Release(TargetIt->Texture);
 
 			TargetIt.RemoveCurrent();
 		}
@@ -150,10 +188,7 @@ void FCompositeRenderTargetPool::Empty()
 
 	for (auto TargetIt = AssignedTargets.CreateIterator(); TargetIt; ++TargetIt)
 	{
-		TObjectPtr<UTextureRenderTarget2D>& Target = TargetIt->Key;
-		Target->ReleaseResource();
-		Target->ConditionalBeginDestroy();
-		Target = nullptr;
+		UE::Composite::Private::Release(TargetIt->Key);
 	}
 
 	AssignedTargets.Empty();
@@ -208,9 +243,7 @@ void FCompositeRenderTargetPool::Tick(float /*DeltaSeconds*/)
 					TObjectPtr<UTextureRenderTarget2D>& Target = TargetIt->Texture;
 					UE_LOGFMT(LogComposite, Verbose, "Releasing stale pooled render target {Name}", Target->GetFName());
 
-					Target->ReleaseResource();
-					Target->ConditionalBeginDestroy();
-					Target = nullptr;
+					UE::Composite::Private::Release(Target);
 
 					TargetIt.RemoveCurrent();
 
diff --git a/Engine/Plugins/Experimental/Compositing/Composite/Source/Composite/Private/CompositeRenderTargetPool.h b/Engine/Plugins/Experimental/Compositing/Composite/Source/Composite/Private/CompositeRenderTargetPool.h
--- a/Engine/Plugins/Experimental/Compositing/Composite/Source/Composite/Private/CompositeRenderTargetPool.h
+++ b/Engine/Plugins/Experimental/Compositing/Composite/Source/Composite/Private/CompositeRenderTargetPool.h
@@ -62,6 +62,9 @@ public:
 	*/
 	void ReleaseTarget(TObjectPtr<UTextureRenderTarget2D>& InTarget);
 
+	/** Release previously assigned render targets to specified UObject.*/
+	void ReleaseAssigneeTargets(UObject* InAssignee);
+
 	/** Destroy available pool render targets. */
 	void EmptyPool(bool bInForceCollectGarbage = true);
 
diff --git a/Engine/Plugins/Experimental/Compositing/Composite/Source/Composite/Private/Layers/CompositeLayerPlate.cpp b/Engine/Plugins/Experimental/Compositing/Composite/Source/Composite/Private/Layers/CompositeLayerPlate.cpp
--- a/Engine/Plugins/Experimental/Compositing/Composite/Source/Composite/Private/Layers/CompositeLayerPlate.cpp
+++ b/Engine/Plugins/Experimental/Compositing/Composite/Source/Composite/Private/Layers/CompositeLayerPlate.cpp
@@ -284,13 +284,9 @@ void UCompositeLayerPlate::PostEditChangeProperty(FPropertyChangedEvent& Propert
 
 		if (bChangeWasRemoval)
 		{
-			if (IsValid(TextureProcessed) && GetValidPreprocessingPassNum() == 0)
+			if (GetValidPreprocessingPassNum() == 0)
 			{
-				TextureProcessed->ConditionalBeginDestroy();
-				TextureProcessed->ReleaseResource();
-				TextureProcessed = nullptr;
-
-				CollectGarbage(GARBAGE_COLLECTION_KEEPFLAGS);
+				FCompositeRenderTargetPool::Get().ReleaseTarget(TextureProcessed);
 			}
 		}
 		else if (PropertyChangedEvent.ChangeType == EPropertyChangeType::ArrayAdd)
diff --git a/Engine/Plugins/Experimental/Compositing/Composite/Source/Composite/Public/CompositeActor.h b/Engine/Plugins/Experimental/Compositing/Composite/Source/Composite/Public/CompositeActor.h
--- a/Engine/Plugins/Experimental/Compositing/Composite/Source/Composite/Public/CompositeActor.h
+++ b/Engine/Plugins/Experimental/Compositing/Composite/Source/Composite/Public/CompositeActor.h
@@ -66,6 +66,7 @@ public:
 	UE_API virtual void PreDuplicateFromRoot(FObjectDuplicationParameters& DupParams) override;
 	UE_API virtual void PreEditChange(FProperty* PropertyAboutToChange) override;
 	UE_API virtual void PostEditChangeChainProperty(struct FPropertyChangedChainEvent& PropertyChangedChainEvent) override;
+	UE_API virtual void PostTransacted(const FTransactionObjectEvent& TransactionEvent) override;
 	UE_API virtual void SetIsTemporarilyHiddenInEditor(bool bIsHidden) override;
 	UE_API virtual bool ShouldTickIfViewportsOnly() const override;
 #endif // WITH_EDITOR
