Composite: Added option to preserve vignetting after key, now enabled by default.

diff --git a/Engine/Plugins/Experimental/Compositing/Composite/Shaders/Private/CompositeDeltaKeyer.usf b/Engine/Plugins/Experimental/Compositing/Composite/Shaders/Private/CompositeDeltaKeyer.usf
--- a/Engine/Plugins/Experimental/Compositing/Composite/Shaders/Private/CompositeDeltaKeyer.usf
+++ b/Engine/Plugins/Experimental/Compositing/Composite/Shaders/Private/CompositeDeltaKeyer.usf
@@ -26,6 +26,7 @@ float4 Params1;
 uint ScreenType;
 uint Visualization;
 uint bInvertAlpha;
+uint bPreserveVignetteAfterKey;
 
 float2 UVToScreenPos(float2 UV, float2 Extent)
 {
@@ -84,13 +85,12 @@ void MainCS(
 
 	float4 Result = Foreground;
 
-	BRANCH
+	float2 ScreenPos = UVToScreenPos(UV, Input_Extent);
+	float AspectRatio = Output_ViewportSize.y * Output_ViewportSizeInverse.x;
+	float VignetteMask = Vignette(ScreenPos, AspectRatio, DevignetteStrength);
+
 	if(DevignetteStrength > ZeroEpsilon)
 	{
-		float2 ScreenPos = UVToScreenPos(UV, Input_Extent);
-		float AspectRatio = Output_ViewportSize.y * Output_ViewportSizeInverse.x;
-		float VignetteMask = Vignette(ScreenPos, AspectRatio, DevignetteStrength);
-
 		Result.rgb *= rcp(VignetteMask);
 	}
 
@@ -126,6 +126,11 @@ void MainCS(
 		Result.rgb = Luv_to_RGB(Foreground_Luv);
 	}
 
+	if(bPreserveVignetteAfterKey > 0 && DevignetteStrength > ZeroEpsilon)
+	{
+		Result.rgb *= VignetteMask;
+	}
+
 	// Apply final alpha threshold
 	Result.a = Linearstep(AlphaMin, AlphaMax, Result.a);
 	
diff --git a/Engine/Plugins/Experimental/Compositing/Composite/Source/Composite/Private/Passes/CompositePassDeltaKeyer.cpp b/Engine/Plugins/Experimental/Compositing/Composite/Source/Composite/Private/Passes/CompositePassDeltaKeyer.cpp
--- a/Engine/Plugins/Experimental/Compositing/Composite/Source/Composite/Private/Passes/CompositePassDeltaKeyer.cpp
+++ b/Engine/Plugins/Experimental/Compositing/Composite/Source/Composite/Private/Passes/CompositePassDeltaKeyer.cpp
@@ -35,6 +35,7 @@ class FCompositePassDeltaKeyerShader : public FGlobalShader
 		SHADER_PARAMETER(FVector4f, Params1)
 		SHADER_PARAMETER(uint32, ScreenType)
 		SHADER_PARAMETER(uint32, Visualization)
+		SHADER_PARAMETER(uint32, bPreserveVignetteAfterKey)
 		SHADER_PARAMETER(uint32, bInvertAlpha)
 	END_SHADER_PARAMETER_STRUCT()
 
@@ -93,6 +94,7 @@ namespace UE
 				PassParameters->Params1 = Params1;
 				PassParameters->ScreenType = static_cast<uint32>(ScreenType);
 				PassParameters->Visualization = static_cast<uint32>(Visualization);
+				PassParameters->bPreserveVignetteAfterKey = static_cast<uint32>(bPreserveVignetteAfterKey);
 				PassParameters->bInvertAlpha = static_cast<uint32>(bInvertAlpha);
 	
 				bool bUseCleanPlate = false;
@@ -130,6 +132,7 @@ namespace UE
 			FVector4f Params0;
 			FVector4f Params1;
 			ECompositeDeltaKeyerVisualization Visualization;
+			bool bPreserveVignetteAfterKey;
 			bool bInvertAlpha;
 			TWeakObjectPtr<UTexture> CleanPlateWeakPtr;
 		};
@@ -146,6 +149,8 @@ UCompositePassDeltaKeyer::UCompositePassDeltaKeyer(const FObjectInitializer& Obj
 	, BlueWeight(0.5f)
 	, AlphaThreshold(0.0f, 1.0f)
 	, DespillStrength(1.0f)
+	, DevignetteStrength(0.0f)
+	, bPreserveVignetteAfterKey(true)
 	, Visualization(ECompositeDeltaKeyerVisualization::Key)
 	, bInvertAlpha(false)
 { 
@@ -162,6 +167,7 @@ bool UCompositePassDeltaKeyer::GetProxy(const UE::CompositeCore::FPassInputDecl&
 	Proxy->Params0 = FVector4f(RedWeight, GreenWeight, BlueWeight, DespillStrength);
 	Proxy->Params1 = FVector4f(AlphaThreshold.X, AlphaThreshold.Y, DevignetteStrength, 0.0f);
 	Proxy->ScreenType = ScreenType;
+	Proxy->bPreserveVignetteAfterKey = bPreserveVignetteAfterKey;
 	Proxy->bInvertAlpha = bInvertAlpha;
 	Proxy->Visualization = Visualization;
 	Proxy->CleanPlateWeakPtr = CleanPlate;
diff --git a/Engine/Plugins/Experimental/Compositing/Composite/Source/Composite/Public/Passes/CompositePassDeltaKeyer.h b/Engine/Plugins/Experimental/Compositing/Composite/Source/Composite/Public/Passes/CompositePassDeltaKeyer.h
--- a/Engine/Plugins/Experimental/Compositing/Composite/Source/Composite/Public/Passes/CompositePassDeltaKeyer.h
+++ b/Engine/Plugins/Experimental/Compositing/Composite/Source/Composite/Public/Passes/CompositePassDeltaKeyer.h
@@ -88,6 +88,10 @@ public:
 	UPROPERTY(EditAnywhere, BlueprintReadWrite, Category = "", meta = (UIMin = 0.0f, UIMax = 1.0f))
 	float DevignetteStrength;
 
+	/** When enabled, we undo devignetting before outputting the keyed plate. */
+	UPROPERTY(EditAnywhere, BlueprintReadWrite, Category = "")
+	bool bPreserveVignetteAfterKey;
+
 	/** Vizualize the alpha key or fill. */
 	UPROPERTY(EditAnywhere, BlueprintReadWrite, Category = "")
 	ECompositeDeltaKeyerVisualization Visualization;
