Composite: Fix lens distortion pass.

* Primary issue was that the pass always needed to write to an output render target even if distortion is not active.
* Added support for engine panini LUT generation during pre-processing.
* Change default to undistort since that will be the primary path for passes going into the scene.

diff --git a/Engine/Plugins/Experimental/Compositing/Composite/Source/Composite/Private/Passes/CompositePassDistortion.cpp b/Engine/Plugins/Experimental/Compositing/Composite/Source/Composite/Private/Passes/CompositePassDistortion.cpp
--- a/Engine/Plugins/Experimental/Compositing/Composite/Source/Composite/Private/Passes/CompositePassDistortion.cpp
+++ b/Engine/Plugins/Experimental/Compositing/Composite/Source/Composite/Private/Passes/CompositePassDistortion.cpp
@@ -61,23 +61,34 @@ public:
 		const FScreenPassTexture& Input = Inputs[0].Texture;
 		FScreenPassRenderTarget Output = Inputs.OverrideOutput;
 
+		// If the override output is provided, it means that this is the last pass in post processing.
+		if (!Output.IsValid())
+		{
+			Output = CreateOutputRenderTarget(GraphBuilder, InView, PassContext.OutputViewRect, Input.Texture->Desc, TEXT("DistortionCompositePass"));
+		}
+
 		FLensDistortionLUT ViewDistortionLUT = LensDistortion::GetLUTUnsafe(InView);
 		bool bValidLensDistortion = ViewDistortionLUT.IsEnabled();
 
-		if (!bValidLensDistortion && LensDistortionSVE)
+		if (!bValidLensDistortion)
 		{
-			// Explicitely render view distortion LUT if existing view one doesn't yet exist, i.e. in pre-processing
-			bValidLensDistortion = LensDistortionSVE->RenderViewDistortionLUT(GraphBuilder, InView.ViewActor.ActorUniqueId, ViewDistortionLUT);
-		}
+			// Explicitely render view distortion LUT if existing view one doesn't yet exist, i.e. in a pre-processing context.
+			if (InView.Family->EngineShowFlags.LensDistortion && FPaniniProjectionConfig::IsEnabledByCVars())
+			{
+				const FPaniniProjectionConfig PaniniProjection = FPaniniProjectionConfig::ReadCVars();
 
-		if (bValidLensDistortion)
-		{
-			// If the override output is provided, it means that this is the last pass in post processing.
-			if (!Output.IsValid())
+				checkSlow(InView.bIsViewInfo);
+				ViewDistortionLUT = PaniniProjection.GenerateLUTPassesUnsafe(GraphBuilder, InView);
+				bValidLensDistortion = ViewDistortionLUT.IsEnabled();
+			}
+			else if (LensDistortionSVE)
 			{
-				Output = CreateOutputRenderTarget(GraphBuilder, InView, PassContext.OutputViewRect, Input.Texture->Desc, TEXT("DistortionCompositePass"));
+				bValidLensDistortion = LensDistortionSVE->RenderViewDistortionLUT(GraphBuilder, InView.ViewActor.ActorUniqueId, ViewDistortionLUT);
 			}
+		}
 
+		if (bValidLensDistortion)
+		{
 			const FScreenPassTextureViewportParameters InputParameters = GetScreenPassTextureViewportParameters(FScreenPassTextureViewport(Input));
 			const FScreenPassTextureViewportParameters OutputParameters = GetScreenPassTextureViewportParameters(FScreenPassTextureViewport(Output));
 
@@ -129,13 +140,13 @@ public:
 			);
 
 			Metadata.bDistorted = (Distortion == ECompositeDistortion::Distort);
-			
-			return UE::CompositeCore::FPassTexture{ MoveTemp(Output), Metadata };
 		}
 		else
 		{
-			return UE::CompositeCore::FPassTexture{ Output.IsValid() ? Output : Input, Metadata };
+			AddDrawTexturePass(GraphBuilder, InView, Input, Output);
 		}
+
+		return UE::CompositeCore::FPassTexture{ MoveTemp(Output), Metadata };
 	}
 
 	TSharedPtr<FLensDistortionSceneViewExtension, ESPMode::ThreadSafe> LensDistortionSVE;
@@ -145,7 +156,7 @@ public:
 UCompositePassDistortion::UCompositePassDistortion(const FObjectInitializer& ObjectInitializer)
 	: Super(ObjectInitializer)
 {
-	Distortion = ECompositeDistortion::Distort;
+	Distortion = ECompositeDistortion::Undistort;
 }
 
 UCompositePassDistortion::~UCompositePassDistortion() = default;
diff --git a/Engine/Plugins/Experimental/Compositing/Composite/Source/Composite/Public/Passes/CompositePassDistortion.h b/Engine/Plugins/Experimental/Compositing/Composite/Source/Composite/Public/Passes/CompositePassDistortion.h
--- a/Engine/Plugins/Experimental/Compositing/Composite/Source/Composite/Public/Passes/CompositePassDistortion.h
+++ b/Engine/Plugins/Experimental/Compositing/Composite/Source/Composite/Public/Passes/CompositePassDistortion.h
@@ -13,8 +13,8 @@
 UENUM(BlueprintType)
 enum class ECompositeDistortion : uint8
 {
+	Undistort,
 	Distort,
-	Undistort
 };
 
 /** Pass to apply an Distortion transform. */
