Composite: Plate layer refactor and addition of scene passes.

* Feature primarily useful to undistort footage going into the scene.
* New mode toggle to let users samples the CRP or texture.
* Fixed Material & Distortion passes activate state function.
* Additional improvements: constness on pool methods, distortion handling in capture layer types, etc.

#virtualized

diff --git a/Engine/Plugins/Experimental/Compositing/Composite/Config/BaseComposite.ini b/Engine/Plugins/Experimental/Compositing/Composite/Config/BaseComposite.ini
--- a/Engine/Plugins/Experimental/Compositing/Composite/Config/BaseComposite.ini
+++ b/Engine/Plugins/Experimental/Compositing/Composite/Config/BaseComposite.ini
@@ -25,9 +25,12 @@
 +ClassRedirects=(OldName="/Script/Composite.CompositeFXAAPass", NewName="/Script/Composite.CompositePassFXAA")
 
 +PropertyRedirects=(OldName="/Script/Composite.CompositeActor.CompositeMeshActors", NewName="/Script/Composite.CompositeLayerPlate.CompositeMeshes")
-+PropertyRedirects=(OldName="/Script/Composite.CompositeLayerPlate.Passes", NewName="/Script/Composite.CompositeLayerPlate.Postprocessing")
-+PropertyRedirects=(OldName="/Script/Composite.CompositeLayerPlate.Subpasses", NewName="/Script/Composite.CompositeLayerPlate.Postprocessing")
++PropertyRedirects=(OldName="/Script/Composite.CompositeLayerPlate.Passes", NewName="/Script/Composite.CompositeLayerPlate.LayerPasses")
++PropertyRedirects=(OldName="/Script/Composite.CompositeLayerPlate.Subpasses", NewName="/Script/Composite.CompositeLayerPlate.LayerPasses")
 +PropertyRedirects=(OldName="/Script/Composite.CompositeLayerPlate.MediaTexture", NewName="/Script/Composite.CompositeLayerPlate.Texture")
 
++PropertyRedirects=(OldName="/Script/Composite.CompositeLayerPlate.Preprocessing", NewName="/Script/Composite.CompositeLayerPlate.MediaPasses")
++PropertyRedirects=(OldName="/Script/Composite.CompositeLayerPlate.Postprocessing", NewName="/Script/Composite.CompositeLayerPlate.LayerPasses")
++PropertyRedirects=(OldName="/Script/Composite.CompositeLayerSceneCapture.Postprocessing", NewName="/Script/Composite.CompositeLayerSceneCapture.LayerPasses")
 
 
diff --git a/Engine/Plugins/Experimental/Compositing/Composite/Source/Composite/Private/CompositeActor.cpp b/Engine/Plugins/Experimental/Compositing/Composite/Source/Composite/Private/CompositeActor.cpp
--- a/Engine/Plugins/Experimental/Compositing/Composite/Source/Composite/Private/CompositeActor.cpp
+++ b/Engine/Plugins/Experimental/Compositing/Composite/Source/Composite/Private/CompositeActor.cpp
@@ -436,13 +436,12 @@ void ACompositeActor::PostEditChangeChainProperty(struct FPropertyChangedChainEv
 
 			if (PropertyChangedChainEvent.PropertyChain.Num() <= 1)
 			{
-				// We only activate the holdout composite pipeline if the plate layer is below the main render
 				for (int32 Index = 0; Index < CompositeLayers.Num(); ++Index)
 				{
 					if (UCompositeLayerPlate* PlateLayer = Cast<UCompositeLayerPlate>(CompositeLayers[Index]))
 					{
-						const bool bHoldoutComposite = Index > MainRenderIndex;
-						PlateLayer->SetHoldoutComposite(bHoldoutComposite);
+						// For user convenience when dragging layers, we disable the composite mesh pipeline if the plate layer runs after the main render.
+						PlateLayer->SetPlateMode(Index > MainRenderIndex ? ECompositePlateMode::CompositeMesh : ECompositePlateMode::Texture);
 					}
 				}
 			}
diff --git a/Engine/Plugins/Experimental/Compositing/Composite/Source/Composite/Private/CompositeRenderTargetPool.cpp b/Engine/Plugins/Experimental/Compositing/Composite/Source/Composite/Private/CompositeRenderTargetPool.cpp
--- a/Engine/Plugins/Experimental/Compositing/Composite/Source/Composite/Private/CompositeRenderTargetPool.cpp
+++ b/Engine/Plugins/Experimental/Compositing/Composite/Source/Composite/Private/CompositeRenderTargetPool.cpp
@@ -44,7 +44,7 @@ FCompositeRenderTargetPool& FCompositeRenderTargetPool::Get()
 	return *PoolInstance;
 }
 
-TObjectPtr<UTextureRenderTarget2D> FCompositeRenderTargetPool::AcquireTarget(UObject* InAssignee, const FIntPoint& InSize, ETextureRenderTargetFormat InFormat)
+TObjectPtr<UTextureRenderTarget2D> FCompositeRenderTargetPool::AcquireTarget(const UObject* InAssignee, const FIntPoint& InSize, ETextureRenderTargetFormat InFormat)
 {
 	TObjectPtr<UTextureRenderTarget2D> RenderTarget = nullptr;
 
@@ -90,7 +90,7 @@ TObjectPtr<UTextureRenderTarget2D> FCompositeRenderTargetPool::AcquireTarget(UOb
 	return RenderTarget;
 }
 
-bool FCompositeRenderTargetPool::ConditionalAcquireTarget(UObject* InAssignee, TObjectPtr<UTextureRenderTarget2D>& InOutTarget, const FIntPoint& InSize, ETextureRenderTargetFormat InFormat)
+bool FCompositeRenderTargetPool::ConditionalAcquireTarget(const UObject* InAssignee, TObjectPtr<UTextureRenderTarget2D>& InOutTarget, const FIntPoint& InSize, ETextureRenderTargetFormat InFormat)
 {
 	if (IsValid(InOutTarget)
 		&& InOutTarget->SizeX == InSize.X
@@ -273,7 +273,7 @@ void FCompositeRenderTargetPool::Tick(float /*DeltaSeconds*/)
 
 	for (auto TargetIt = AssignedTargets.CreateIterator(); TargetIt; ++TargetIt)
 	{
-		const TWeakObjectPtr<UObject>& Assignee = TargetIt->Value;
+		const TWeakObjectPtr<const UObject>& Assignee = TargetIt->Value;
 		if (!Assignee.IsValid())
 		{
 			TargetIt.RemoveCurrent();
diff --git a/Engine/Plugins/Experimental/Compositing/Composite/Source/Composite/Private/CompositeRenderTargetPool.h b/Engine/Plugins/Experimental/Compositing/Composite/Source/Composite/Private/CompositeRenderTargetPool.h
--- a/Engine/Plugins/Experimental/Compositing/Composite/Source/Composite/Private/CompositeRenderTargetPool.h
+++ b/Engine/Plugins/Experimental/Compositing/Composite/Source/Composite/Private/CompositeRenderTargetPool.h
@@ -40,7 +40,7 @@ public:
 	* @param InFormat Render target texture format.
 	* @return Render target object.
 	*/
-	TObjectPtr<UTextureRenderTarget2D> AcquireTarget(UObject* InAssignee, const FIntPoint& InSize, ETextureRenderTargetFormat InFormat = ETextureRenderTargetFormat::RTF_RGBA16f);
+	TObjectPtr<UTextureRenderTarget2D> AcquireTarget(const UObject* InAssignee, const FIntPoint& InSize, ETextureRenderTargetFormat InFormat = ETextureRenderTargetFormat::RTF_RGBA16f);
 
 	/**
 	* Conditionally acquire new or recyled render target, only if the size & format are mismatched with the current render target.
@@ -51,7 +51,7 @@ public:
 	* @param InFormat Render target texture format.
 	* @return True if the render target object was updated.
 	*/
-	bool ConditionalAcquireTarget(UObject* InAssignee, TObjectPtr<UTextureRenderTarget2D>& InOutTarget, const FIntPoint& InSize, ETextureRenderTargetFormat InFormat = ETextureRenderTargetFormat::RTF_RGBA16f);
+	bool ConditionalAcquireTarget(const UObject* InAssignee, TObjectPtr<UTextureRenderTarget2D>& InOutTarget, const FIntPoint& InSize, ETextureRenderTargetFormat InFormat = ETextureRenderTargetFormat::RTF_RGBA16f);
 
 	/** Check if the render target was created by the pool and is currently assigned. */
 	bool IsAssignedTarget(const TObjectPtr<UTextureRenderTarget2D>& InTarget) const;
@@ -126,6 +126,6 @@ private:
 	TMap<const FRenderTargetDesc, TArray<FPooledRenderTarget>> PooledTargets;
 
 	/** Assigned render targets (to be returned into the pool). */
-	TMap<TObjectPtr<UTextureRenderTarget2D>, TWeakObjectPtr<UObject>> AssignedTargets;
+	TMap<TObjectPtr<UTextureRenderTarget2D>, TWeakObjectPtr<const UObject>> AssignedTargets;
 };
 
diff --git a/Engine/Plugins/Experimental/Compositing/Composite/Source/Composite/Private/Layers/CompositeLayerPlate.cpp b/Engine/Plugins/Experimental/Compositing/Composite/Source/Composite/Private/Layers/CompositeLayerPlate.cpp
--- a/Engine/Plugins/Experimental/Compositing/Composite/Source/Composite/Private/Layers/CompositeLayerPlate.cpp
+++ b/Engine/Plugins/Experimental/Compositing/Composite/Source/Composite/Private/Layers/CompositeLayerPlate.cpp
@@ -13,6 +13,7 @@
 #include "Engine/TextureRenderTarget2D.h"
 #include "Materials/Material.h"
 #include "Materials/MaterialInstanceDynamic.h"
+#include "Passes/CompositePassDistortion.h"
 #include "UObject/ConstructorHelpers.h"
 
 #if WITH_EDITOR
@@ -74,9 +75,12 @@ UCompositeLayerPlate::UCompositeLayerPlate(const FObjectInitializer& ObjectIniti
 	static ConstructorHelpers::FObjectFinder<UTexture2D> DefaultMediaTexture = TEXT("/Composite/Textures/T_Composite_SMPTE_Color_Bars_16x9.T_Composite_SMPTE_Color_Bars_16x9");
 	
 	Texture = Cast<UTexture>(DefaultMediaTexture.Object);
-	bHoldoutComposite = true;
+	PlateMode = ECompositePlateMode::CompositeMesh;
 	Operation = ECompositeCoreMergeOp::Over;
 
+	// Default undistort pass for plate media going into the scene
+	ScenePasses.Add(CreateDefaultSubobject<UCompositePassDistortion>(TEXT("CompositeUndistortScenePass")));
+
 #if WITH_EDITOR
 	// If we are not the class default object...
 	if (!HasAnyFlags(RF_ClassDefaultObject))
@@ -108,11 +112,46 @@ void UCompositeLayerPlate::Tick(float DeltaTime)
 	
 	if (IsValid(Subsystem))
 	{
-		if (UE::CompositeCore::IsRegisterPrimitivesOnTickEnabled() && bHoldoutComposite)
+		// TODO: We will need to separate the holdout feature from the custom render pass in the scene view extension, so that one can be used without the other.
+		if (UE::CompositeCore::IsRegisterPrimitivesOnTickEnabled())
 		{
 			Subsystem->RegisterPrimitives(GetPrimitives());
 		}
 	}
+
+	{
+		// Constantly check if passes have changed to manage render targets & texture bindings on composite meshes.
+		const int32 NumValidMediaPasses = GetValidPassesNum(MediaPasses);
+		const int32 NumValidScenePasses = GetValidPassesNum(ScenePasses);
+		bool bPassesHaveChanged = false;
+
+		if (NumValidMediaPasses != CachedValidMediaPasses)
+		{
+			if (NumValidMediaPasses == 0)
+			{
+				FCompositeRenderTargetPool::Get().ReleaseTarget(MediaRenderTarget);
+			}
+
+			CachedValidMediaPasses = NumValidMediaPasses;
+			bPassesHaveChanged = true;
+		}
+
+		if (NumValidScenePasses != CachedValidScenePasses)
+		{
+			if (NumValidScenePasses == 0)
+			{
+				FCompositeRenderTargetPool::Get().ReleaseTarget(SceneRenderTarget);
+			}
+
+			CachedValidScenePasses = NumValidScenePasses;
+			bPassesHaveChanged = true;
+		}
+
+		if (bPassesHaveChanged)
+		{
+			UpdateCompositeMeshes();
+		}
+	}
 }
 
 ETickableTickType UCompositeLayerPlate::GetTickableTickType() const
@@ -145,7 +184,7 @@ void UCompositeLayerPlate::PostInitProperties()
 		
 		if (ensure(IsValid(CompositeActor)))
 		{
-			CompositeActor->OnEnabledStateChange.BindUObject(this, &UCompositeLayerPlate::HandleHoldoutCompositeStateChange);
+			CompositeActor->OnEnabledStateChange.BindUObject(this, &UCompositeLayerPlate::HandleEnabledStateChange);
 		}
 	}
 }
@@ -156,11 +195,6 @@ void UCompositeLayerPlate::PostLoad()
 
 	if (!HasAnyFlags(RF_ClassDefaultObject))
 	{
-		if (IsValid(Texture) && GetValidPreprocessingPassNum() > 0)
-		{
-			FCompositeRenderTargetPool::Get().ConditionalAcquireTarget(this, TextureProcessed, GetRenderResolution());
-		}
-
 		UpdateCompositeMeshes();
 	}
 }
@@ -169,7 +203,8 @@ void UCompositeLayerPlate::BeginDestroy()
 {
 	Super::BeginDestroy();
 
-	FCompositeRenderTargetPool::Get().ReleaseTarget(TextureProcessed);
+	FCompositeRenderTargetPool::Get().ReleaseTarget(MediaRenderTarget);
+	FCompositeRenderTargetPool::Get().ReleaseTarget(SceneRenderTarget);
 
 #if WITH_EDITOR
 	// If we are not the class default object...
@@ -197,21 +232,23 @@ void UCompositeLayerPlate::SetCompositeMeshes(TArray<AActor*> InCompositeMeshes)
 	UpdateCompositeMeshes();
 }
 
-void UCompositeLayerPlate::SetHoldoutComposite(bool bInHoldoutComposite)
+void UCompositeLayerPlate::SetPlateMode(ECompositePlateMode InPlateMode)
 {
-	if (bHoldoutComposite != bInHoldoutComposite)
+	if (PlateMode != InPlateMode)
 	{
-		bHoldoutComposite = bInHoldoutComposite;
-		
-		HandleHoldoutCompositeStateChange(bHoldoutComposite, UWorld::GetSubsystem<UCompositeCoreSubsystem>(GetWorld()));
+		PlateMode = InPlateMode;
 	}
 }
 
 UTexture* UCompositeLayerPlate::GetCompositeTexture() const
 {
-	if (IsValid(TextureProcessed) && GetValidPreprocessingPassNum() > 0)
+	if (CachedValidScenePasses > 0)
+	{		
+		return GetOrCreateRenderTarget(SceneRenderTarget).Get();
+	}
+	else if (CachedValidMediaPasses > 0)
 	{
-		return TextureProcessed.Get();
+		return GetOrCreateRenderTarget(MediaRenderTarget).Get();
 	}
 	else if (IsValid(Texture))
 	{
@@ -224,20 +261,7 @@ UTexture* UCompositeLayerPlate::GetCompositeTexture() const
 #if WITH_EDITOR
 bool UCompositeLayerPlate::CanEditChange(const FProperty* InProperty) const
 {
-	bool bIsEditable = Super::CanEditChange(InProperty);
-
-	if (bIsEditable && InProperty)
-	{
-		const FName PropertyName = InProperty->GetFName();
-
-		// Always return true for enabled toggle
-		if (PropertyName == GET_MEMBER_NAME_CHECKED(ThisClass, CompositeMeshes))
-		{
-			bIsEditable = bHoldoutComposite;
-		}
-	}
-
-	return bIsEditable;
+	return Super::CanEditChange(InProperty);
 }
 
 void UCompositeLayerPlate::PreEditChange(FProperty* PropertyThatWillChange)
@@ -269,48 +293,9 @@ void UCompositeLayerPlate::PostEditChangeProperty(FPropertyChangedEvent& Propert
 
 	const FName PropertyName = PropertyChangedEvent.GetMemberPropertyName();
 
-	if (PropertyName == GET_MEMBER_NAME_CHECKED(ThisClass, Texture))
-	{
-		UpdateCompositeMeshes();
-	}
-	if (PropertyName == GET_MEMBER_NAME_CHECKED(ThisClass, bHoldoutComposite))
-	{
-		HandleHoldoutCompositeStateChange(bHoldoutComposite, UWorld::GetSubsystem<UCompositeCoreSubsystem>(GetWorld()));
-	}
-	else if (PropertyName == GET_MEMBER_NAME_CHECKED(ThisClass, Preprocessing))
-	{
-		const bool bChangeWasRemoval = PropertyChangedEvent.ChangeType == EPropertyChangeType::ArrayRemove
-			|| PropertyChangedEvent.ChangeType == EPropertyChangeType::ArrayClear;
-
-		if (bChangeWasRemoval)
-		{
-			if (GetValidPreprocessingPassNum() == 0)
-			{
-				FCompositeRenderTargetPool::Get().ReleaseTarget(TextureProcessed);
-			}
-		}
-		else if (PropertyChangedEvent.ChangeType == EPropertyChangeType::ArrayAdd)
-		{
-			// Always add elements at the start of the array (top of the stack)
-			Algo::Rotate(Preprocessing, Preprocessing.Num() - 1);
-		}
-		else if (PropertyChangedEvent.ChangeType == EPropertyChangeType::ValueSet)
-		{
-			if (IsValid(Texture) && !IsValid(TextureProcessed))
-			{
-				FCompositeRenderTargetPool::Get().ConditionalAcquireTarget(this, TextureProcessed, GetRenderResolution());
-			}
-		}
-
-		UpdateCompositeMeshes();
-	}
-	else if (PropertyName == GET_MEMBER_NAME_CHECKED(ThisClass, Postprocessing))
+	if (PropertyName == GET_MEMBER_NAME_CHECKED(ThisClass, bIsEnabled))
 	{
-		if (PropertyChangedEvent.ChangeType == EPropertyChangeType::ArrayAdd)
-		{
-			// Always add elements at the start of the array (top of the stack)
-			Algo::Rotate(Postprocessing, Postprocessing.Num() - 1);
-		}
+		HandleEnabledStateChange(bIsEnabled, UWorld::GetSubsystem<UCompositeCoreSubsystem>(GetWorld()));
 	}
 	else if (PropertyName == GET_MEMBER_NAME_CHECKED(ThisClass, CompositeMeshes))
 	{
@@ -342,8 +327,6 @@ void UCompositeLayerPlate::PostEditChangeProperty(FPropertyChangedEvent& Propert
 					}
 				}
 			}
-
-			UpdateCompositeMeshes();
 		}
 
 		UCompositeCoreSubsystem* Subsystem = UWorld::GetSubsystem<UCompositeCoreSubsystem>(GetWorld());
@@ -355,6 +338,9 @@ void UCompositeLayerPlate::PostEditChangeProperty(FPropertyChangedEvent& Propert
 
 		PreEditCompositeMeshes.Reset();
 	}
+
+	// For simplicity, we make sure to update composite meshes on any property change.
+	UpdateCompositeMeshes();
 }
 
 void UCompositeLayerPlate::PostTransacted(const FTransactionObjectEvent& InTransactionEvent)
@@ -380,118 +366,143 @@ void UCompositeLayerPlate::PostTransacted(const FTransactionObjectEvent& InTrans
 				Subsystem->RegisterPrimitives(GetPrimitives());
 			}
 		}
-		else if (ChangedProperties.Contains(GET_MEMBER_NAME_CHECKED(ThisClass, Texture)))
+	}
+}
+#endif
+
+
+int32 UCompositeLayerPlate::FindLastValidPassIndex(TArrayView<const TObjectPtr<UCompositePassBase>> InPasses) const
+{
+	int32 LastValidPassIndex = INDEX_NONE;
+	
+	// Iterate backwards so that lower passes in the UI are executed first
+	for (int32 PassIndex = InPasses.Num() - 1; PassIndex >= 0; --PassIndex)
+	{
+		const TObjectPtr<UCompositePassBase>& Pass = InPasses[PassIndex];
+
+		if (IsValid(Pass) && Pass->IsActive())
 		{
-			HandleHoldoutCompositeStateChange(bHoldoutComposite, UWorld::GetSubsystem<UCompositeCoreSubsystem>(GetWorld()));
+			LastValidPassIndex = PassIndex;
 		}
 	}
+
+	return LastValidPassIndex;
 }
-#endif
 
-bool UCompositeLayerPlate::GetProxy(FTraversalContext& InContext, FSceneRenderingBulkObjectAllocator& InFrameAllocator, FCompositeCorePassProxy*& OutProxy) const
+UE::CompositeCore::ResourceId UCompositeLayerPlate::AddPreprocessingPasses(
+	FTraversalContext& InContext,
+	FSceneRenderingBulkObjectAllocator& InFrameAllocator,
+	TArrayView<const TObjectPtr<UCompositePassBase>> InPasses,
+	UE::CompositeCore::ResourceId TextureId,
+	UE::CompositeCore::ResourceId OriginalTextureId,
+	TFunction<TObjectPtr<UTextureRenderTarget2D>()> GetRenderTargetFn
+) const
 {
 	using namespace UE::CompositeCore;
 
-	TOptional<ResourceId> TextureId;
-	TOptional<ResourceId> TextureProcessedId;
-
-	if (IsValid(Texture))
+	if (InPasses.IsEmpty())
 	{
-		FExternalTexture& ExternalTexture = InContext.ExternalTextures.AddDefaulted_GetRef();
-		ExternalTexture.Texture = Texture;
-		ExternalTexture.Metadata.bDistorted = true;
-
-		TextureId = EXTERNAL_RANGE_START_ID + static_cast<ResourceId>(InContext.ExternalTextures.Num() - 1);
+		return TextureId;
 	}
 
-	if (TextureId.IsSet() && IsValid(TextureProcessed))
+	ResourceId OutputTextureId = TextureId;
+	bool bIsFirstValidPass = true;
+	const int32 LastValidPassIndex = FindLastValidPassIndex(InPasses);
+
+	// Iterate backwards so that lower passes in the UI are executed first
+	for (int32 PassIndex = InPasses.Num() - 1; PassIndex >= 0; --PassIndex)
 	{
-		// Iterate backwards so that lower passes in the UI are executed first
-		const int32 FirstPassIndex = Preprocessing.Num() - 1;
+		const TObjectPtr<UCompositePassBase>& Prepass = InPasses[PassIndex];
 
-		TOptional<int32> FirstValidPassIndex;
-		TOptional<int32> LastValidPassIndex;
-		for (int32 PassIndex = FirstPassIndex; PassIndex >= 0; --PassIndex)
+		if (!IsValid(Prepass) || !Prepass->IsActive())
 		{
-			const TObjectPtr<UCompositePassBase>& Prepass = Preprocessing[PassIndex];
+			continue;
+		}
 
-			if (IsValid(Prepass) && Prepass->IsActive())
-			{
-				if (!FirstValidPassIndex.IsSet())
-				{
-					FirstValidPassIndex = PassIndex;
-				}
+		FPassInputDecl PassInput;
 
-				LastValidPassIndex = PassIndex;
-			}
+		if (bIsFirstValidPass)
+		{
+			PassInput.Set<FPassExternalResourceDesc>({ TextureId });
+			bIsFirstValidPass = false;
 		}
-
-		for (int32 PassIndex = FirstPassIndex; PassIndex >= 0; --PassIndex)
+		else
 		{
-			const TObjectPtr<UCompositePassBase>& Prepass = Preprocessing[PassIndex];
+			// Default input from previous pass
+			PassInput.Set<FPassInternalResourceDesc>({});
+		}
 
-			if (!IsValid(Prepass) || !Prepass->IsActive())
+		FCompositeCorePassProxy* PrepassProxy;
+		if (Prepass->GetProxy(PassInput, InFrameAllocator, PrepassProxy))
+		{
+			if (PassIndex == LastValidPassIndex)
 			{
-				continue;
-			}
+				FExternalTexture& ExternalTexture = InContext.ExternalTextures.AddDefaulted_GetRef();
+				ExternalTexture.Texture = GetRenderTargetFn();
+				ExternalTexture.Metadata.bDistorted = true;
 
-			FPassInputDecl PassInput;
-			if (PassIndex == FirstValidPassIndex)
-			{
-				PassInput.Set<FPassExternalResourceDesc>({ TextureId.GetValue() });
-			}
-			else
-			{
-				// Default input from previous pass
-				PassInput.Set<FPassInternalResourceDesc>({});
+				OutputTextureId = EXTERNAL_RANGE_START_ID + static_cast<ResourceId>(InContext.ExternalTextures.Num() - 1);
+
+				PrepassProxy->DeclareOutputOverride(OutputTextureId);
 			}
 
-			FCompositeCorePassProxy* PrepassProxy;
-			if (Prepass->GetProxy(PassInput, InFrameAllocator, PrepassProxy))
-			{
-				if (PassIndex == LastValidPassIndex)
-				{
-					FExternalTexture& ExternalTexture = InContext.ExternalTextures.AddDefaulted_GetRef();
-					ExternalTexture.Texture = TextureProcessed;
-					ExternalTexture.Metadata.bDistorted = true;
+			InContext.PreprocessingPasses.FindOrAdd(OriginalTextureId).Add(PrepassProxy);
+		}
+	}
 
-					check(!TextureProcessedId.IsSet());
-					TextureProcessedId = EXTERNAL_RANGE_START_ID + static_cast<ResourceId>(InContext.ExternalTextures.Num() - 1);
+	return OutputTextureId;
+}
 
-					PrepassProxy->DeclareOutputOverride(TextureProcessedId.GetValue());
-				}
+bool UCompositeLayerPlate::GetProxy(FTraversalContext& InContext, FSceneRenderingBulkObjectAllocator& InFrameAllocator, FCompositeCorePassProxy*& OutProxy) const
+{
+	using namespace UE::CompositeCore;
 
-				TArray<const FCompositeCorePassProxy*>& PrepassProxies = InContext.PreprocessingPasses.FindOrAdd(TextureId.GetValue());
-				PrepassProxies.Add(PrepassProxy);
-			}
-		}
-	}
-	
-	ResourceId InputId = BUILT_IN_EMPTY_ID;
-	if (bHoldoutComposite && (GetValidPrimitivesNum() > 0))
-	{
-		InputId = BUILT_IN_CRP_ID;
-	}
-	else if (TextureProcessedId.IsSet())
+	ResourceId TextureId = BUILT_IN_EMPTY_ID;
+
+	if (IsValid(Texture))
 	{
-		InputId = TextureProcessedId.GetValue();
+		// Register source texture into the frame work context (so that post-processing passes can refer to it).
+		FExternalTexture& ExternalTexture = InContext.ExternalTextures.AddDefaulted_GetRef();
+		ExternalTexture.Texture = Texture;
+		ExternalTexture.Metadata.bDistorted = true;
+		TextureId = EXTERNAL_RANGE_START_ID + static_cast<ResourceId>(InContext.ExternalTextures.Num() - 1);
+		
+		const ResourceId OriginalTextureId = TextureId;
+
+		// Add media texture pre-processing passes
+		TextureId = AddPreprocessingPasses(InContext, InFrameAllocator, MediaPasses, TextureId, OriginalTextureId,
+			[this]() -> TObjectPtr<UTextureRenderTarget2D>
+			{
+				return GetOrCreateRenderTarget(MediaRenderTarget);
+			}
+		);
+
+		// Add scene-only passes. We ignore returned resource since we don't refer to it in post-processing.
+		AddPreprocessingPasses(InContext, InFrameAllocator, ScenePasses, TextureId, OriginalTextureId,
+			[this]() -> TObjectPtr<UTextureRenderTarget2D>
+			{
+				return GetOrCreateRenderTarget(SceneRenderTarget);
+			}
+		);
 	}
-	else if (TextureId.IsSet())
+
+	if (PlateMode == ECompositePlateMode::CompositeMesh && GetValidPrimitivesNum() > 0)
 	{
-		InputId = TextureId.GetValue();
+		// Prioritize the built-in custom render pass when enabled
+		TextureId = BUILT_IN_CRP_ID;
 	}
 
 	FPassInputDecl PassInput;
-	PassInput.Set<FPassExternalResourceDesc>({ InputId });
+	PassInput.Set<FPassExternalResourceDesc>({ TextureId });
 
-	AddChildPasses(PassInput, InContext, InFrameAllocator, Postprocessing);
+	AddChildPasses(PassInput, InContext, InFrameAllocator, LayerPasses);
 
 	FPassInputDeclArray Inputs;
 	Inputs.SetNum(FixedNumLayerInputs);
 	Inputs[0] = PassInput;
 	Inputs[1] = GetDefaultSecondInput(InContext);
 
-	OutProxy = InFrameAllocator.Create<FMergePassProxy>(MoveTemp(Inputs), GetMergeOperation(InContext), TEXT("Plate"), ELensDistortionHandling::Disabled);
+	OutProxy = InFrameAllocator.Create<FMergePassProxy>(MoveTemp(Inputs), GetMergeOperation(InContext), TEXT("Plate"));
 	return true;
 }
 
@@ -541,15 +552,15 @@ int32 UCompositeLayerPlate::GetValidPrimitivesNum() const
 	return Count;
 }
 
-int32 UCompositeLayerPlate::GetValidPreprocessingPassNum() const
+int32 UCompositeLayerPlate::GetValidPassesNum(TArrayView<const TObjectPtr<UCompositePassBase>> InPasses) const
 {
 	int32 Count = 0;
 
-	for (int32 Index = 0; Index < Preprocessing.Num(); ++Index)
+	for (int32 Index = 0; Index < InPasses.Num(); ++Index)
 	{
-		const TObjectPtr<UCompositePassBase>& Prepass = Preprocessing[Index];
+		const TObjectPtr<UCompositePassBase>& Pass = InPasses[Index];
 
-		if (IsValid(Prepass) && Prepass->IsActive())
+		if (IsValid(Pass) && Pass->IsActive())
 		{
 			++Count;
 		}
@@ -643,14 +654,20 @@ void UCompositeLayerPlate::UpdatePrimitiveComponent(UPrimitiveComponent& InPrimi
 				*/
 				MID = InPrimitiveComponent.CreateAndSetMaterialInstanceDynamic(ElementIndex);
 			}
-
+			
+			MID->SetScalarParameterValue(TEXT("UseCompositeProjection"), static_cast<float>(PlateMode == ECompositePlateMode::CompositeMesh));
 			MID->SetTextureParameterValue(TEXT("CompositeTexture"), GetCompositeTexture());
 		}
 	}
 }
 
-void UCompositeLayerPlate::HandleHoldoutCompositeStateChange(bool bInEnabled, UCompositeCoreSubsystem* InSubsystem) const
+void UCompositeLayerPlate::HandleEnabledStateChange(bool bInEnabled, UCompositeCoreSubsystem* InSubsystem) const
 {
+	if (CompositeMeshes.IsEmpty())
+	{
+		return;
+	}
+
 	if (bInEnabled)
 	{
 		constexpr bool bRegisterOnly = true;
@@ -658,6 +675,7 @@ void UCompositeLayerPlate::HandleHoldoutCompositeStateChange(bool bInEnabled, UC
 
 		if (IsValid(InSubsystem))
 		{
+			// Rely on the scene view extension to mark primivites as holdout, and render them separately.
 			InSubsystem->RegisterPrimitives(GetPrimitives());
 		}
 	}
@@ -695,4 +713,14 @@ bool UCompositeLayerPlate::IsCompositeMeshActorAlreadyInUse(AActor* InCompositeM
 	return false;
 }
 
+TObjectPtr<UTextureRenderTarget2D> UCompositeLayerPlate::GetOrCreateRenderTarget(TObjectPtr<UTextureRenderTarget2D>& InRenderTarget) const
+{
+	if (!IsValid(InRenderTarget))
+	{
+		FCompositeRenderTargetPool::Get().ConditionalAcquireTarget(this, InRenderTarget, GetRenderResolution());
+	}
+
+	return InRenderTarget;
+}
+
 #undef LOCTEXT_NAMESPACE
diff --git a/Engine/Plugins/Experimental/Compositing/Composite/Source/Composite/Private/Layers/CompositeLayerSceneCapture.cpp b/Engine/Plugins/Experimental/Compositing/Composite/Source/Composite/Private/Layers/CompositeLayerSceneCapture.cpp
--- a/Engine/Plugins/Experimental/Compositing/Composite/Source/Composite/Private/Layers/CompositeLayerSceneCapture.cpp
+++ b/Engine/Plugins/Experimental/Compositing/Composite/Source/Composite/Private/Layers/CompositeLayerSceneCapture.cpp
@@ -52,12 +52,12 @@ void UCompositeLayerSceneCapture::PostEditChangeProperty(FPropertyChangedEvent&
 	{
 		SetCustomRenderPass(bCustomRenderPass);
 	}
-	else if (PropertyName == GET_MEMBER_NAME_CHECKED(ThisClass, Postprocessing))
+	else if (PropertyName == GET_MEMBER_NAME_CHECKED(ThisClass, LayerPasses))
 	{
 		if (PropertyChangedEvent.ChangeType == EPropertyChangeType::ArrayAdd)
 		{
 			// Always add elements at the start of the array (top of the stack)
-			Algo::Rotate(Postprocessing, Postprocessing.Num() - 1);
+			Algo::Rotate(LayerPasses, LayerPasses.Num() - 1);
 		}
 	}
 }
@@ -84,21 +84,21 @@ bool UCompositeLayerSceneCapture::GetProxy(FTraversalContext& InContext, FSceneR
 	FExternalTexture& ExternalTexture = InContext.ExternalTextures.AddDefaulted_GetRef();
 	ExternalTexture.Texture = CaptureComponent->TextureTarget;
 	ExternalTexture.Metadata.bInvertedAlpha = true;
-	ExternalTexture.Metadata.bDistorted = false;
+	ExternalTexture.Metadata.bDistorted = CaptureComponent->ShowFlags.LensDistortion;
 	ExternalTexture.Metadata.Encoding = UE::CompositeCore::EEncoding::Linear;
 
 	const int32 TexIndex = InContext.ExternalTextures.Num() - 1;
 	FPassInputDecl PassInput;
 	PassInput.Set<FPassExternalResourceDesc>({ EXTERNAL_RANGE_START_ID + TexIndex });
 
-	AddChildPasses(PassInput, InContext, InFrameAllocator, Postprocessing);
+	AddChildPasses(PassInput, InContext, InFrameAllocator, LayerPasses);
 
 	FPassInputDeclArray Inputs;
 	Inputs.SetNum(FixedNumLayerInputs);
 	Inputs[0] = PassInput;
 	Inputs[1] = GetDefaultSecondInput(InContext);
 
-	OutProxy = InFrameAllocator.Create<FMergePassProxy>(MoveTemp(Inputs), GetMergeOperation(InContext), TEXT("SceneCapture"));
+	OutProxy = InFrameAllocator.Create<FMergePassProxy>(MoveTemp(Inputs), GetMergeOperation(InContext), TEXT("SceneCapture"), ELensDistortionHandling::Disabled);
 	return true;
 }
 
diff --git a/Engine/Plugins/Experimental/Compositing/Composite/Source/Composite/Private/Layers/CompositeLayerShadowReflection.cpp b/Engine/Plugins/Experimental/Compositing/Composite/Source/Composite/Private/Layers/CompositeLayerShadowReflection.cpp
--- a/Engine/Plugins/Experimental/Compositing/Composite/Source/Composite/Private/Layers/CompositeLayerShadowReflection.cpp
+++ b/Engine/Plugins/Experimental/Compositing/Composite/Source/Composite/Private/Layers/CompositeLayerShadowReflection.cpp
@@ -82,7 +82,7 @@ bool UCompositeLayerShadowReflection::GetProxy(FTraversalContext& InContext, FSc
 			FExternalTexture& ExternalTexture = InContext.ExternalTextures.AddDefaulted_GetRef();
 			ExternalTexture.Texture = Component->TextureTarget;
 			ExternalTexture.Metadata.bInvertedAlpha = true;
-			ExternalTexture.Metadata.bDistorted = false;
+			ExternalTexture.Metadata.bDistorted = Component->ShowFlags.LensDistortion;
 			ExternalTexture.Metadata.Encoding = UE::CompositeCore::EEncoding::Linear;
 
 			int32 TexIndex = InContext.ExternalTextures.Num() - 1;
diff --git a/Engine/Plugins/Experimental/Compositing/Composite/Source/Composite/Private/Passes/CompositePassDistortion.cpp b/Engine/Plugins/Experimental/Compositing/Composite/Source/Composite/Private/Passes/CompositePassDistortion.cpp
--- a/Engine/Plugins/Experimental/Compositing/Composite/Source/Composite/Private/Passes/CompositePassDistortion.cpp
+++ b/Engine/Plugins/Experimental/Compositing/Composite/Source/Composite/Private/Passes/CompositePassDistortion.cpp
@@ -2,8 +2,10 @@
 
 #include "Passes/CompositePassDistortion.h"
 
+#include "Camera/CameraActor.h"
 #include "Camera/CameraComponent.h"
 #include "CameraCalibrationSubsystem.h"
+#include "CompositeActor.h"
 #include "LensDistortionSceneViewExtension.h"
 #include "ScreenPass.h"
 #include "SystemTextures.h"
@@ -161,20 +163,49 @@ UCompositePassDistortion::UCompositePassDistortion(const FObjectInitializer& Obj
 
 UCompositePassDistortion::~UCompositePassDistortion() = default;
 
-bool UCompositePassDistortion::GetProxy(const UE::CompositeCore::FPassInputDecl& InputDecl, FSceneRenderingBulkObjectAllocator& InFrameAllocator, FCompositeCorePassProxy*& OutProxy) const
+void UCompositePassDistortion::PostInitProperties()
 {
-	TSharedPtr<FLensDistortionSceneViewExtension, ESPMode::ThreadSafe> SVE;
-	UCameraCalibrationSubsystem* SubSystem = GEngine ? GEngine->GetEngineSubsystem<UCameraCalibrationSubsystem>() : nullptr;
-	if (!SubSystem)
+	Super::PostInitProperties();
+
+	UCameraCalibrationSubsystem* Subsystem = GEngine ? GEngine->GetEngineSubsystem<UCameraCalibrationSubsystem>() : nullptr;
+	if (Subsystem)
 	{
-		return false;
+		LensDistortionSVE = Subsystem->GetSceneViewExtension();
 	}
+}
+
+bool UCompositePassDistortion::IsActive() const
+{
+	bool bIsDistortionActive = false;
 
-	FDistortionCompositePassProxy* Proxy = InFrameAllocator.Create<FDistortionCompositePassProxy>(UE::CompositeCore::FPassInputDeclArray{ InputDecl });
-	Proxy->Distortion = Distortion;
-	Proxy->LensDistortionSVE = SubSystem->GetSceneViewExtension();
+	if (Super::IsActive())
+	{
+		static IConsoleVariable* PaniniCVar = IConsoleManager::Get().FindConsoleVariable(TEXT("r.LensDistortion.Panini.D"));
+		if (LensDistortionSVE && LensDistortionSVE->HasDistortionState_AnyThread())
+		{
+			bIsDistortionActive = true;
+		}
+		else if (PaniniCVar->GetFloat() > 0.01f)
+		{
+			bIsDistortionActive = true;
+		}
+	}
+
+	return bIsDistortionActive;
+}
+
+bool UCompositePassDistortion::GetProxy(const UE::CompositeCore::FPassInputDecl& InputDecl, FSceneRenderingBulkObjectAllocator& InFrameAllocator, FCompositeCorePassProxy*& OutProxy) const
+{
+	if (IsActive())
+	{
+		FDistortionCompositePassProxy* Proxy = InFrameAllocator.Create<FDistortionCompositePassProxy>(UE::CompositeCore::FPassInputDeclArray{ InputDecl });
+		Proxy->Distortion = Distortion;
+		Proxy->LensDistortionSVE = LensDistortionSVE;
+
+		OutProxy = Proxy;
+		return true;
+	}
 
-	OutProxy = Proxy;
-	return true;
+	return false;
 }
 
diff --git a/Engine/Plugins/Experimental/Compositing/Composite/Source/Composite/Private/Passes/CompositePassMaterial.cpp b/Engine/Plugins/Experimental/Compositing/Composite/Source/Composite/Private/Passes/CompositePassMaterial.cpp
--- a/Engine/Plugins/Experimental/Compositing/Composite/Source/Composite/Private/Passes/CompositePassMaterial.cpp
+++ b/Engine/Plugins/Experimental/Compositing/Composite/Source/Composite/Private/Passes/CompositePassMaterial.cpp
@@ -53,9 +53,14 @@ UCompositePassMaterial::UCompositePassMaterial(const FObjectInitializer& ObjectI
 
 UCompositePassMaterial::~UCompositePassMaterial() = default;
 
+bool UCompositePassMaterial::IsActive() const
+{
+	return Super::IsActive() && IsValid(PostProcessMaterial) && PostProcessMaterial->IsPostProcessMaterial();
+}
+
 bool UCompositePassMaterial::GetProxy(const UE::CompositeCore::FPassInputDecl& InputDecl, FSceneRenderingBulkObjectAllocator& InFrameAllocator, FCompositeCorePassProxy*& OutProxy) const
 {
-	if (IsValid(PostProcessMaterial) && PostProcessMaterial->IsPostProcessMaterial())
+	if (IsActive())
 	{
 		FCompositePassMaterialProxy* Proxy = InFrameAllocator.Create<FCompositePassMaterialProxy>(UE::CompositeCore::FPassInputDeclArray{ InputDecl });
 		Proxy->Material = PostProcessMaterial;
diff --git a/Engine/Plugins/Experimental/Compositing/Composite/Source/Composite/Public/Layers/CompositeLayerPlate.h b/Engine/Plugins/Experimental/Compositing/Composite/Source/Composite/Public/Layers/CompositeLayerPlate.h
--- a/Engine/Plugins/Experimental/Compositing/Composite/Source/Composite/Public/Layers/CompositeLayerPlate.h
+++ b/Engine/Plugins/Experimental/Compositing/Composite/Source/Composite/Public/Layers/CompositeLayerPlate.h
@@ -16,6 +16,14 @@ class UTexture;
 class FObjectPreSaveContext;
 class UTextureRenderTarget2D;
 
+/** Plate sampling & rendering mode. */
+UENUM(BlueprintType)
+enum class ECompositePlateMode : uint8
+{
+	Texture = 0 UMETA(ToolTip = "Sample the texture directly in 2D screen space."),
+	CompositeMesh = 1 UMETA(ToolTip = "Sample the composite meshes rendered in a built-in custom render pass (default). Automatic fallback to Texture."),
+};
+
 /**
 * Convenience layer for media texture integration, with additional processing options.
 * 
@@ -79,15 +87,15 @@ public:
 	UFUNCTION(BlueprintSetter)
 	UE_API void SetCompositeMeshes(TArray<AActor*> InCompositeMeshes);
 
-	/** Get whether the holdout composite pipeline is active. */
+	/** Get the plate sampling & rendering mode */
 	UFUNCTION(BlueprintGetter)
-	UE_API bool IsHoldoutComposite() const { return bHoldoutComposite; }
+	UE_API ECompositePlateMode GetPlateMode() const { return PlateMode; }
 
-	/** Enable/disable the holdout composite pipeline. */
+	/** Set the plate sampling & rendering mode */
 	UFUNCTION(BlueprintSetter)
-	UE_API void SetHoldoutComposite(bool bInHoldoutComposite);
+	UE_API void SetPlateMode(ECompositePlateMode InPlateMode);
 
-	/** Returns the media texture, optionally pre-processed into a render target. */
+	/** Returns the media texture, optionally pre-processed with layer & scene-only passes into a render target. */
 	UFUNCTION(BlueprintCallable, Category = "Composite")
 	UTexture* GetCompositeTexture() const;
 
@@ -101,18 +109,34 @@ public:
 	UPROPERTY(EditAnywhere, BlueprintReadWrite, Category = "", meta = (AllowedClasses="/Script/Engine.Texture2D, /Script/MediaAssets.MediaTexture"))
 	TObjectPtr<UTexture> Texture;
 
-	/** Texture pre-processing passes, applied at the start of rendering and updating a render target. */
-	UPROPERTY(EditAnywhere, BlueprintReadWrite, Instanced, Category = "", meta = (DisplayName = "Pre-Processing", DisallowedClasses = "/Script/Composite.CompositePassMaterial"))
-	TArray<TObjectPtr<UCompositePassBase>> Preprocessing;
+	/**
+	* Media texture pre-processing passes, applied before rendering.
+	* 
+	* The texture is replaced with an internal render target.
+	*/
+	UPROPERTY(EditAnywhere, BlueprintReadWrite, Instanced, Category = "", meta = (DisallowedClasses = "/Script/Composite.CompositePassMaterial, /Script/Composite.CompositePassDistortion, /Script/Composite.CompositePassFXAA, /Script/Composite.CompositePassSMAA"))
+	TArray<TObjectPtr<UCompositePassBase>> MediaPasses;
 
-	/** Sub-passes applied during post-processing on the texture or holdout composite render. */
-	UPROPERTY(EditAnywhere, BlueprintReadWrite, Instanced, Category = "", meta = (DisplayName = "Post-Processing"))
-	TArray<TObjectPtr<UCompositePassBase>> Postprocessing;
+	/**
+	* Passes applied after media passes, onto composite meshes only.
+	* 
+	* The texture is replaced with an internal render target.
+	* 
+	* NOTE: Set the plate "Mode" to "Texture" in order to skip these scene passes during final compositing in post-processing.
+	* We now have an automatic scene pass to undistort a plate projected onto composite meshes. Using Texture mode can therefore
+	* avoid the automatic undistort -> distort pipeline.
+	*/
+	UPROPERTY(BlueprintReadWrite, Instanced, Category = "", meta = (AllowedClasses = "/Script/Composite.CompositePassDistortion"))
+	TArray<TObjectPtr<UCompositePassBase>> ScenePasses;
+
+	/** Passes applied during final compositing in post-processing. */
+	UPROPERTY(EditAnywhere, BlueprintReadWrite, Instanced, Category = "", meta = (DisallowedClasses = "/Script/Composite.CompositePassDistortion"))
+	TArray<TObjectPtr<UCompositePassBase>> LayerPasses;
 
 private:
-	/** Enable holdout composite pipeline on the composite meshes. (Automatically disabled when the layer is moved above the main render.) */
-	UPROPERTY(EditAnywhere, BlueprintGetter = IsHoldoutComposite, BlueprintSetter = SetHoldoutComposite, Category = "")
-	bool bHoldoutComposite;
+	/** The plate sampling & rendering mode: direct texture sampling or separate composite mesh render. */
+	UPROPERTY(EditAnywhere, BlueprintGetter = GetPlateMode, BlueprintSetter = SetPlateMode, Category = "", AdvancedDisplay, meta = (AllowPrivateAccess = true, DisplayName = "Mode"))
+	ECompositePlateMode PlateMode;
 
 private:
 	/** Update the (media) texture on all composite meshes. We also add a user asset data to track changes. */
@@ -121,19 +145,35 @@ private:
 	/** Update the (media) texture on a primitive component. */
 	void UpdatePrimitiveComponent(UPrimitiveComponent& InPrimitiveComponent) const;
 
-	/** Handle state change of the owner composite actor. */
-	void HandleHoldoutCompositeStateChange(bool bInEnabled, UCompositeCoreSubsystem* InSubsystem) const;
+	/** Handle enabled state change to (de)activate the built-in custom render pass for composite meshes. */
+	void HandleEnabledStateChange(bool bInEnabled, UCompositeCoreSubsystem* InSubsystem) const;
 
 	/** Convenience function verify if the composite mesh actor is already in use by another layer. */
 	bool IsCompositeMeshActorAlreadyInUse(AActor* InCompositeMeshActor) const;
 
+	/** Convenience function to get or create the specified render target at composite actor resolution. */
+	TObjectPtr<UTextureRenderTarget2D> GetOrCreateRenderTarget(TObjectPtr<UTextureRenderTarget2D>& InRenderTarget) const;
+
 	/** Get number of valid primitives. */
 	int32 GetValidPrimitivesNum() const;
 
-	/** Get number of valid pre-processing passes. */
-	int32 GetValidPreprocessingPassNum() const;
+	/** Get number of valid passes. */
+	int32 GetValidPassesNum(TArrayView<const TObjectPtr<UCompositePassBase>> InPasses) const;
 
 private:
+	/** Convenience function to find the last valid pass. */
+	int32 FindLastValidPassIndex(TArrayView<const TObjectPtr<UCompositePassBase>> InPasses) const;
+
+	/** Convenience function to register pre-processing child passes. */
+	UE::CompositeCore::ResourceId AddPreprocessingPasses(
+		FTraversalContext& InContext,
+		FSceneRenderingBulkObjectAllocator& InFrameAllocator,
+		TArrayView<const TObjectPtr<UCompositePassBase>> InPasses,
+		UE::CompositeCore::ResourceId TextureId,
+		UE::CompositeCore::ResourceId OriginalTextureId,
+		TFunction<TObjectPtr<UTextureRenderTarget2D>()> GetRenderTargetFn
+	) const;
+
 #if WITH_EDITOR
 	/** Pre-edit list of registered composite meshes. */
 	TArray<TObjectPtr<AActor>> PreEditCompositeMeshes;
@@ -149,9 +189,19 @@ private:
 	void OnPostSaveWorld(UWorld* InWorld, FObjectPostSaveContext ObjectSaveContext);
 #endif
 
-	/** Processed texture render target. Note: Should be replaced by pool-managed RT. */
-	UPROPERTY(Transient)
-	TObjectPtr<UTextureRenderTarget2D> TextureProcessed;
+	/** Processed media texture render target.*/
+	UPROPERTY(Transient, NonTransactional)
+	mutable TObjectPtr<UTextureRenderTarget2D> MediaRenderTarget;
+
+	/** Processed media texture with additional passes only for scene composite meshes.*/
+	UPROPERTY(Transient, NonTransactional)
+	mutable TObjectPtr<UTextureRenderTarget2D> SceneRenderTarget;
+
+	/** Cached number of valid media passes. */
+	int32 CachedValidMediaPasses = 0;
+
+	/** Cached number of valid scene passes. */
+	int32 CachedValidScenePasses = 0;
 };
 
 #undef UE_API
diff --git a/Engine/Plugins/Experimental/Compositing/Composite/Source/Composite/Public/Layers/CompositeLayerProcessing.h b/Engine/Plugins/Experimental/Compositing/Composite/Source/Composite/Public/Layers/CompositeLayerProcessing.h
--- a/Engine/Plugins/Experimental/Compositing/Composite/Source/Composite/Public/Layers/CompositeLayerProcessing.h
+++ b/Engine/Plugins/Experimental/Compositing/Composite/Source/Composite/Public/Layers/CompositeLayerProcessing.h
@@ -32,7 +32,7 @@ public:
 #endif
 
 	/** Sub-passes applied onto the previous layer. */
-	UPROPERTY(EditAnywhere, BlueprintReadWrite, Instanced, Category = "")
+	UPROPERTY(EditAnywhere, BlueprintReadWrite, Instanced, Category = "", meta = (DisallowedClasses = "/Script/Composite.CompositePassDistortion"))
 	TArray<TObjectPtr<UCompositePassBase>> LayerPasses;
 };
 
diff --git a/Engine/Plugins/Experimental/Compositing/Composite/Source/Composite/Public/Layers/CompositeLayerSceneCapture.h b/Engine/Plugins/Experimental/Compositing/Composite/Source/Composite/Public/Layers/CompositeLayerSceneCapture.h
--- a/Engine/Plugins/Experimental/Compositing/Composite/Source/Composite/Public/Layers/CompositeLayerSceneCapture.h
+++ b/Engine/Plugins/Experimental/Compositing/Composite/Source/Composite/Public/Layers/CompositeLayerSceneCapture.h
@@ -51,8 +51,8 @@ public:
 	UE_API void SetCustomRenderPass(bool bInIsFastRenderPass);
 
 	/** Sub-passes applied during post-processing. */
-	UPROPERTY(EditAnywhere, BlueprintReadWrite, Instanced, Category = "", meta = (DisplayName = "Post-Processing", DisplayAfter = "bCustomRenderPass"))
-	TArray<TObjectPtr<UCompositePassBase>> Postprocessing;
+	UPROPERTY(EditAnywhere, BlueprintReadWrite, Instanced, Category = "", meta = (DisallowedClasses = "/Script/Composite.CompositePassDistortion", DisplayAfter = "bCustomRenderPass"))
+	TArray<TObjectPtr<UCompositePassBase>> LayerPasses;
 
 private:
 	/** Convenience function to access the actor-managed scene capture component. */
diff --git a/Engine/Plugins/Experimental/Compositing/Composite/Source/Composite/Public/Passes/CompositePassDistortion.h b/Engine/Plugins/Experimental/Compositing/Composite/Source/Composite/Public/Passes/CompositePassDistortion.h
--- a/Engine/Plugins/Experimental/Compositing/Composite/Source/Composite/Public/Passes/CompositePassDistortion.h
+++ b/Engine/Plugins/Experimental/Compositing/Composite/Source/Composite/Public/Passes/CompositePassDistortion.h
@@ -10,6 +10,8 @@
 
 #define UE_API COMPOSITE_API
 
+class FLensDistortionSceneViewExtension;
+
 UENUM(BlueprintType)
 enum class ECompositeDistortion : uint8
 {
@@ -29,12 +31,22 @@ public:
 	/** Destructor */
 	UE_API ~UCompositePassDistortion();
 
+	//~ Begin UObject Interface
+	virtual void PostInitProperties() override;
+	//~ End UObject Interface
+
+	virtual bool IsActive() const override;
+
 	UE_API virtual bool GetProxy(const UE::CompositeCore::FPassInputDecl& InputDecl, FSceneRenderingBulkObjectAllocator& InFrameAllocator, FCompositeCorePassProxy*& OutProxy) const override;
 
 public:
 	/** Distortion to apply on input. */
 	UPROPERTY(EditAnywhere, BlueprintReadWrite, Category = "")
 	ECompositeDistortion Distortion;
+
+private:
+	/** Cached lens distortion scene view extension. */
+	TSharedPtr<FLensDistortionSceneViewExtension, ESPMode::ThreadSafe> LensDistortionSVE;
 };
 
 #undef UE_API
diff --git a/Engine/Plugins/Experimental/Compositing/Composite/Source/Composite/Public/Passes/CompositePassMaterial.h b/Engine/Plugins/Experimental/Compositing/Composite/Source/Composite/Public/Passes/CompositePassMaterial.h
--- a/Engine/Plugins/Experimental/Compositing/Composite/Source/Composite/Public/Passes/CompositePassMaterial.h
+++ b/Engine/Plugins/Experimental/Compositing/Composite/Source/Composite/Public/Passes/CompositePassMaterial.h
@@ -23,6 +23,8 @@ public:
 	/** Destructor */
 	UE_API ~UCompositePassMaterial();
 
+	UE_API virtual bool IsActive() const override;
+
 	UE_API virtual bool GetProxy(const UE::CompositeCore::FPassInputDecl& InputDecl, FSceneRenderingBulkObjectAllocator& InFrameAllocator, FCompositeCorePassProxy*& OutProxy) const override;
 
 	UE_API virtual bool NeedsSceneTextures() const override;
