Composite: Fix small issues with previous CL

* Missed CompositeLayerPlate description update
* Plate mode toggle comment on ordering

diff --git a/Engine/Plugins/Experimental/Compositing/Composite/Source/Composite/Private/CompositeActor.cpp b/Engine/Plugins/Experimental/Compositing/Composite/Source/Composite/Private/CompositeActor.cpp
--- a/Engine/Plugins/Experimental/Compositing/Composite/Source/Composite/Private/CompositeActor.cpp
+++ b/Engine/Plugins/Experimental/Compositing/Composite/Source/Composite/Private/CompositeActor.cpp
@@ -441,6 +441,7 @@ void ACompositeActor::PostEditChangeChainProperty(struct FPropertyChangedChainEv
 					if (UCompositeLayerPlate* PlateLayer = Cast<UCompositeLayerPlate>(CompositeLayers[Index]))
 					{
 						// For user convenience when dragging layers, we disable the composite mesh pipeline if the plate layer runs after the main render.
+						// TODO: Invert condition if/when traversal is done in the regular indexed order in the new UI.
 						PlateLayer->SetPlateMode(Index > MainRenderIndex ? ECompositePlateMode::CompositeMesh : ECompositePlateMode::Texture);
 					}
 				}
diff --git a/Engine/Plugins/Experimental/Compositing/Composite/Source/Composite/Public/Layers/CompositeLayerPlate.h b/Engine/Plugins/Experimental/Compositing/Composite/Source/Composite/Public/Layers/CompositeLayerPlate.h
--- a/Engine/Plugins/Experimental/Compositing/Composite/Source/Composite/Public/Layers/CompositeLayerPlate.h
+++ b/Engine/Plugins/Experimental/Compositing/Composite/Source/Composite/Public/Layers/CompositeLayerPlate.h
@@ -27,8 +27,8 @@ enum class ECompositePlateMode : uint8
 /**
 * Convenience layer for media texture integration, with additional processing options.
 * 
-* Note that when "Holdout Composite" is active, the media texture is projected onto composite meshes & rendered separately as a custom render pass.
-* The result is dilated and combined with other layers such as the main render (with composite meshes rendered as holdouts).
+* With composite meshes and default plugin materials, the media texture is projected into the scene & rendered separately as a custom render pass.
+* The result is dilated and combined with other layers such as the main render (where the same composite meshes are rendered as holdouts).
 */
 UCLASS(MinimalAPI, BlueprintType, Blueprintable, EditInlineNew, CollapseCategories, meta = (DisplayName = "Plate"))
 class UCompositeLayerPlate : public UCompositeLayerBase, public FTickableGameObject
