Composite: Adding media filter denoising option.

* Separate pass needed to operate at original resolution.

diff --git a/Engine/Plugins/Experimental/Compositing/Composite/Shaders/Private/CompositeDeltaKeyerCommon.ush b/Engine/Plugins/Experimental/Compositing/Composite/Shaders/Private/CompositeDeltaKeyerCommon.ush
--- a/Engine/Plugins/Experimental/Compositing/Composite/Shaders/Private/CompositeDeltaKeyerCommon.ush
+++ b/Engine/Plugins/Experimental/Compositing/Composite/Shaders/Private/CompositeDeltaKeyerCommon.ush
@@ -279,4 +279,33 @@ float4 LuvAngleKeyer(float3 Foreground_Luv, float3 Background_Luv, float Accepta
 
 	// TODO: Modulate Foreground_Luv.x
 	return float4(Foreground_Luv, Alpha);
-}
\ No newline at end of file
+}
+
+// Branchless 2-value sort
+void Sort2(inout float4 A, inout float4 B)
+{
+    float4 Temp = A;
+    A = min(A, B);
+    B = max(Temp, B);
+}
+
+// Median of 9 colors (https://bertdobbelaere.github.io/median_networks.html#N9L19D7)
+float4 Median9(float4 A0, float4 A1, float4 A2, float4 A3, float4 A4, float4 A5, float4 A6, float4 A7, float4 A8)
+{
+	// For future optimization, the operations on each line could be done in parallel on separate threads or lanes.
+	Sort2(A0, A7); Sort2(A1, A2); Sort2(A3, A5); Sort2(A4, A8);
+	Sort2(A0, A2); Sort2(A1, A5); Sort2(A3, A8); Sort2(A4, A7);
+	Sort2(A0, A3); Sort2(A1, A4); Sort2(A2, A8); Sort2(A5, A7);
+	Sort2(A3, A4); Sort2(A5, A6);
+	Sort2(A2, A5); Sort2(A4, A6);
+	Sort2(A2, A3); Sort2(A4, A5);
+	Sort2(A3, A4);
+
+	return A4;
+}
+
+// Denoising with per-channel median filter of 9 colors
+float4 DenoiseMedianFilter9(const float4 V[9])
+{
+    return Median9(V[0], V[1], V[2], V[3], V[4], V[5], V[6], V[7], V[8]);
+}
diff --git a/Engine/Plugins/Experimental/Compositing/Composite/Shaders/Private/CompositeDenoise.usf b/Engine/Plugins/Experimental/Compositing/Composite/Shaders/Private/CompositeDenoise.usf
new file mode 100644
--- /dev/null
+++ b/Engine/Plugins/Experimental/Compositing/Composite/Shaders/Private/CompositeDenoise.usf
@@ -0,0 +1,104 @@
+// Copyright Epic Games, Inc. All Rights Reserved.
+
+#include "/Engine/Private/Common.ush"
+
+#include "CompositeDeltaKeyerCommon.ush"
+
+#ifndef TILE_SIZE
+	#error TILE_SIZE must be defined
+#endif
+
+#define FILTER_SIZE 1
+#define FILTER_COUNT (1 + 2 * FILTER_SIZE) * (1 + 2 * FILTER_SIZE)
+// Width in texels of the depth tile cached into LDS.
+#define LDS_TILE_WIDTH (TILE_SIZE + 2 * FILTER_SIZE)
+// Total number of texels cached in the scene color tile.
+#define LDS_ARRAY_SIZE (LDS_TILE_WIDTH * LDS_TILE_WIDTH)
+#define THREADGROUP_TOTAL (TILE_SIZE * TILE_SIZE)
+
+Texture2D InputTexture;
+RWTexture2D<float4> RWOutputTexture;
+uint2 Dimensions;
+
+groupshared float4 SharedColor[LDS_ARRAY_SIZE];
+
+static const int2 kOffsets[FILTER_COUNT] = {
+	int2( 0,  0),
+	int2(-1,  0),
+	int2( 1,  0),
+	int2( 0,  1),
+	int2( 0, -1),
+	int2(-1,  1),
+	int2( 1,  1),
+	int2( 1, -1),
+	int2(-1, -1),
+};
+
+// Get the index within the LDS array.
+uint GetTileArrayIndexFromPixelOffset(in uint2 GroupThreadId, in int2 PixelOffset)
+{
+	uint2 TilePos = GroupThreadId + uint2(PixelOffset + FILTER_SIZE);
+	return TilePos.x + TilePos.y * (FILTER_SIZE * 2 + TILE_SIZE);
+}
+
+float4 SampleCachedColors(in uint2 GroupThreadId, int2 PixelOffset = int2(0,0))
+{
+	return SharedColor[GetTileArrayIndexFromPixelOffset(GroupThreadId, PixelOffset)];
+}
+
+
+[numthreads(TILE_SIZE, TILE_SIZE, 1)]
+void MainCS(
+	uint2 GroupId : SV_GroupID,
+	uint2 GroupThreadId : SV_GroupThreadID,
+	uint GroupThreadIndex : SV_GroupIndex,
+	uint2 DispatchThreadId : SV_DispatchThreadID)
+{
+	{
+		// Init LDS, same logic as CompositeCoreDilate & TAA.
+		uint2 GroupTexelOffset = GroupId * uint2(TILE_SIZE, TILE_SIZE) - FILTER_SIZE;
+		// Each warp thread loads extra texels for the filtering neighbors
+		const uint LoadCount = (LDS_ARRAY_SIZE + THREADGROUP_TOTAL - 1) / THREADGROUP_TOTAL;
+		uint LinearGroupThreadId = GroupThreadIndex;
+		
+		UNROLL
+		for (uint i = 0; i < LoadCount; i++)
+		{
+			uint2 TexelLocation =  GroupTexelOffset + uint2(
+				LinearGroupThreadId % LDS_TILE_WIDTH,
+				LinearGroupThreadId / LDS_TILE_WIDTH);
+			TexelLocation = min(TexelLocation, Dimensions - 1);
+
+			if ((LinearGroupThreadId < LDS_ARRAY_SIZE) ||
+				(i != LoadCount - 1) ||
+				(LDS_ARRAY_SIZE % THREADGROUP_TOTAL) == 0)
+			{
+				SharedColor[LinearGroupThreadId] = InputTexture[TexelLocation];
+			}
+			LinearGroupThreadId += THREADGROUP_TOTAL;
+		}
+	}
+
+    // Discard out-of-bound texels
+    if (any(DispatchThreadId >= Dimensions))
+    {
+        return;
+    }
+
+	GroupMemoryBarrierWithGroupSync();
+
+	float4 TexelColor = SampleCachedColors(GroupThreadId);
+
+	{
+		float4 Neighbors[FILTER_COUNT];
+		UNROLL
+		for(int i=0; i<FILTER_COUNT; i++)
+		{
+		    Neighbors[i] = SampleCachedColors(GroupThreadId, kOffsets[i]);
+		}
+
+		TexelColor = DenoiseMedianFilter9(Neighbors);
+	}
+
+	RWOutputTexture[DispatchThreadId] = TexelColor;
+}
diff --git a/Engine/Plugins/Experimental/Compositing/Composite/Source/Composite/Private/Passes/CompositePassDeltaKeyer.cpp b/Engine/Plugins/Experimental/Compositing/Composite/Source/Composite/Private/Passes/CompositePassDeltaKeyer.cpp
--- a/Engine/Plugins/Experimental/Compositing/Composite/Source/Composite/Private/Passes/CompositePassDeltaKeyer.cpp
+++ b/Engine/Plugins/Experimental/Compositing/Composite/Source/Composite/Private/Passes/CompositePassDeltaKeyer.cpp
@@ -11,8 +11,30 @@
 #include "Application/ThrottleManager.h"
 #endif
 
+DECLARE_GPU_STAT_NAMED(FCompositeDenoise, TEXT("Composite.Denoise"));
 DECLARE_GPU_STAT_NAMED(FCompositeDeltaKeyer, TEXT("Composite.DeltaKeyer"));
 
+class FCompositePassDenoiseShader : public FGlobalShader
+{
+	DECLARE_GLOBAL_SHADER(FCompositePassDenoiseShader);
+	SHADER_USE_PARAMETER_STRUCT(FCompositePassDenoiseShader, FGlobalShader);
+
+	BEGIN_SHADER_PARAMETER_STRUCT(FParameters, )
+		SHADER_PARAMETER_RDG_TEXTURE(Texture2D, InputTexture)
+		SHADER_PARAMETER_RDG_TEXTURE_UAV(RWTexture2D<float4>, RWOutputTexture)
+		SHADER_PARAMETER(FIntPoint, Dimensions)
+	END_SHADER_PARAMETER_STRUCT()
+
+	static inline void ModifyCompilationEnvironment(const FGlobalShaderPermutationParameters& Parameters, FShaderCompilerEnvironment& OutEnvironment)
+	{
+		OutEnvironment.SetDefine(TEXT("TILE_SIZE"), ThreadGroupSize);
+
+		FGlobalShader::ModifyCompilationEnvironment(Parameters, OutEnvironment);
+	}
+
+	static const uint32 ThreadGroupSize = 16;
+};
+
 class FCompositePassDeltaKeyerShader : public FGlobalShader
 {
 	DECLARE_GLOBAL_SHADER(FCompositePassDeltaKeyerShader);
@@ -48,8 +70,9 @@ class FCompositePassDeltaKeyerShader : public FGlobalShader
 
 	static const uint32 ThreadGroupSize = 16;
 };
-IMPLEMENT_GLOBAL_SHADER(FCompositePassDeltaKeyerShader, "/Plugin/Composite/Private/CompositeDeltaKeyer.usf", "MainCS", SF_Compute);
 
+IMPLEMENT_GLOBAL_SHADER(FCompositePassDenoiseShader, "/Plugin/Composite/Private/CompositeDenoise.usf", "MainCS", SF_Compute);
+IMPLEMENT_GLOBAL_SHADER(FCompositePassDeltaKeyerShader, "/Plugin/Composite/Private/CompositeDeltaKeyer.usf", "MainCS", SF_Compute);
 
 namespace UE
 {
@@ -60,6 +83,35 @@ namespace UE
 		public:
 			using FCompositeCorePassProxy::FCompositeCorePassProxy;
 
+			FScreenPassTexture AddDenoisePass(FRDGBuilder& GraphBuilder, const FScreenPassTexture& Input, ERHIFeatureLevel::Type FeatureLevel) const
+			{
+				RDG_EVENT_SCOPE_STAT(GraphBuilder, FCompositeDenoise, "Composite.Denoise");
+				RDG_GPU_STAT_SCOPE(GraphBuilder, FCompositeDenoise);
+
+				FGlobalShaderMap* GlobalShaderMap = GetGlobalShaderMap(FeatureLevel);
+				const FIntPoint TextureSize = Input.Texture->Desc.Extent;
+
+				const FRDGTextureDesc OutputDesc = FRDGTextureDesc::Create2D(TextureSize, PF_FloatRGBA, FClearValueBinding::Black, TexCreate_UAV | TexCreate_ShaderResource | TexCreate_RenderTargetable);
+				FScreenPassRenderTarget Output = FScreenPassRenderTarget(GraphBuilder.CreateTexture(OutputDesc, TEXT("CompositeTextureDenoised")), ERenderTargetLoadAction::ENoAction);
+
+				FCompositePassDenoiseShader::FParameters* PassParameters = GraphBuilder.AllocParameters<FCompositePassDenoiseShader::FParameters>();
+				PassParameters->InputTexture = Input.Texture;
+				PassParameters->RWOutputTexture = GraphBuilder.CreateUAV(Output.Texture);
+				PassParameters->Dimensions = TextureSize;
+
+				TShaderMapRef<FCompositePassDenoiseShader> ComputeShader(GlobalShaderMap);
+				FComputeShaderUtils::AddPass(
+					GraphBuilder,
+					RDG_EVENT_NAME("CompositeCore.Denoise (%dx%d)", TextureSize.X, TextureSize.Y),
+					GSupportsEfficientAsyncCompute ? ERDGPassFlags::AsyncCompute : ERDGPassFlags::Compute,
+					ComputeShader,
+					PassParameters,
+					FComputeShaderUtils::GetGroupCount(TextureSize, FCompositePassDenoiseShader::ThreadGroupSize)
+				);
+
+				return Output;
+			}
+
 			FPassTexture Add(FRDGBuilder& GraphBuilder, const FSceneView& InView, const FPassInputArray& Inputs, const FPassContext& PassContext) const override
 			{
 				RDG_EVENT_SCOPE_STAT(GraphBuilder, FCompositeDeltaKeyer, "Composite.DeltaKeyer");
@@ -68,9 +120,14 @@ namespace UE
 				check(ValidateInputs(Inputs));
 
 				const UE::CompositeCore::FResourceMetadata& Metadata = Inputs[0].Metadata;
-				const FScreenPassTexture& Input = Inputs[0].Texture;
+				FScreenPassTexture Input = Inputs[0].Texture;
 				FScreenPassRenderTarget Output = Inputs.OverrideOutput;
 
+				if (DenoiseMethod != ECompositeDenoiseMethod::None)
+				{
+					Input = AddDenoisePass(GraphBuilder, Input, InView.GetFeatureLevel());
+				}
+
 				// If the override output is provided, it means that this is the last pass in post processing.
 				if (!Output.IsValid())
 				{
@@ -133,6 +190,7 @@ namespace UE
 			FVector4f Params1;
 			ECompositeDeltaKeyerVisualization Visualization;
 			bool bPreserveVignetteAfterKey;
+			ECompositeDenoiseMethod DenoiseMethod;
 			bool bInvertAlpha;
 			TWeakObjectPtr<UTexture> CleanPlateWeakPtr;
 		};
@@ -151,6 +209,7 @@ UCompositePassDeltaKeyer::UCompositePassDeltaKeyer(const FObjectInitializer& Obj
 	, DespillStrength(1.0f)
 	, DevignetteStrength(0.0f)
 	, bPreserveVignetteAfterKey(true)
+	, DenoiseMethod(ECompositeDenoiseMethod::None)
 	, Visualization(ECompositeDeltaKeyerVisualization::Key)
 	, bInvertAlpha(false)
 { 
@@ -168,6 +227,7 @@ bool UCompositePassDeltaKeyer::GetProxy(const UE::CompositeCore::FPassInputDecl&
 	Proxy->Params1 = FVector4f(AlphaThreshold.X, AlphaThreshold.Y, DevignetteStrength, 0.0f);
 	Proxy->ScreenType = ScreenType;
 	Proxy->bPreserveVignetteAfterKey = bPreserveVignetteAfterKey;
+	Proxy->DenoiseMethod = DenoiseMethod;
 	Proxy->bInvertAlpha = bInvertAlpha;
 	Proxy->Visualization = Visualization;
 	Proxy->CleanPlateWeakPtr = CleanPlate;
diff --git a/Engine/Plugins/Experimental/Compositing/Composite/Source/Composite/Public/Passes/CompositePassDeltaKeyer.h b/Engine/Plugins/Experimental/Compositing/Composite/Source/Composite/Public/Passes/CompositePassDeltaKeyer.h
--- a/Engine/Plugins/Experimental/Compositing/Composite/Source/Composite/Public/Passes/CompositePassDeltaKeyer.h
+++ b/Engine/Plugins/Experimental/Compositing/Composite/Source/Composite/Public/Passes/CompositePassDeltaKeyer.h
@@ -26,6 +26,15 @@ enum class ECompositeDeltaKeyerScreenType : uint8
 	Blue = 2
 };
 
+/** Background screen color type (red, green or blue).*/
+UENUM(BlueprintType)
+enum class ECompositeDenoiseMethod : uint8
+{
+	None,
+	Median3x3 UMETA(DisplayName = "Median 3x3"),
+};
+
+
 /**
  * Color keyer with clean plate & spill removal support.
 */
@@ -92,6 +101,10 @@ public:
 	UPROPERTY(EditAnywhere, BlueprintReadWrite, Category = "")
 	bool bPreserveVignetteAfterKey;
 
+	/** Denoising method applied before the keyer. */
+	UPROPERTY(EditAnywhere, BlueprintReadWrite, Category = "")
+	ECompositeDenoiseMethod DenoiseMethod;
+
 	/** Vizualize the alpha key or fill. */
 	UPROPERTY(EditAnywhere, BlueprintReadWrite, Category = "")
 	ECompositeDeltaKeyerVisualization Visualization;
