Composite: Fixed invalid denoising enum description in previous  CL.

diff --git a/Engine/Plugins/Experimental/Compositing/Composite/Source/Composite/Public/Passes/CompositePassDeltaKeyer.h b/Engine/Plugins/Experimental/Compositing/Composite/Source/Composite/Public/Passes/CompositePassDeltaKeyer.h
--- a/Engine/Plugins/Experimental/Compositing/Composite/Source/Composite/Public/Passes/CompositePassDeltaKeyer.h
+++ b/Engine/Plugins/Experimental/Compositing/Composite/Source/Composite/Public/Passes/CompositePassDeltaKeyer.h
@@ -26,7 +26,7 @@ enum class ECompositeDeltaKeyerScreenType : uint8
 	Blue = 2
 };
 
-/** Background screen color type (red, green or blue).*/
+/** Denoising method applied before keying. */
 UENUM(BlueprintType)
 enum class ECompositeDenoiseMethod : uint8
 {
