Composure: Added Pass table to Plate details panel, which displays a tree view with all the plate's passes.

diff --git a/Engine/Plugins/Experimental/Compositing/Composite/Source/Composite/Private/CompositeActor.cpp b/Engine/Plugins/Experimental/Compositing/Composite/Source/Composite/Private/CompositeActor.cpp
--- a/Engine/Plugins/Experimental/Compositing/Composite/Source/Composite/Private/CompositeActor.cpp
+++ b/Engine/Plugins/Experimental/Compositing/Composite/Source/Composite/Private/CompositeActor.cpp
@@ -417,12 +417,6 @@ void ACompositeActor::PostEditChangeChainProperty(struct FPropertyChangedChainEv
 				}
 			}
 
-			if (PropertyChangedChainEvent.ChangeType == EPropertyChangeType::ArrayAdd)
-			{
-				// Always add elements at the start of the array (top of the stack)
-				Algo::Rotate(CompositeLayers, CompositeLayers.Num() - 1);
-			}
-
 			for (UCompositeLayerBase* PreEditCompositeLayer : PreEditCompositeLayers)
 			{
 				if (IsValid(PreEditCompositeLayer))
diff --git a/Engine/Plugins/Experimental/Compositing/Composite/Source/Composite/Public/Layers/CompositeLayerPlate.h b/Engine/Plugins/Experimental/Compositing/Composite/Source/Composite/Public/Layers/CompositeLayerPlate.h
--- a/Engine/Plugins/Experimental/Compositing/Composite/Source/Composite/Public/Layers/CompositeLayerPlate.h
+++ b/Engine/Plugins/Experimental/Compositing/Composite/Source/Composite/Public/Layers/CompositeLayerPlate.h
@@ -63,14 +63,14 @@ public:
 	/** Get array of all composite mesh primitives. */
 	UE_API TArray<UPrimitiveComponent*> GetPrimitives() const;
 
-protected:
+public:
 	//~ Begin UObject Interface
 	UE_API virtual void PostInitProperties() override;
 	UE_API virtual void PostLoad() override;
 	UE_API virtual void BeginDestroy() override;
 
 #if WITH_EDITOR
-	UE_API virtual bool CanEditChange(const FProperty* InProperty) const;
+	UE_API virtual bool CanEditChange(const FProperty* InProperty) const override;
 	UE_API virtual void PreEditChange(FProperty* PropertyThatWillChange) override;
 	UE_API virtual void PostEditChangeProperty(FPropertyChangedEvent& PropertyChangedEvent) override;
 	UE_API virtual void PostTransacted(const FTransactionObjectEvent& InTransactionEvent) override;
diff --git a/Engine/Plugins/Experimental/Compositing/Composite/Source/Composite/Public/Passes/CompositePassColorGrade.h b/Engine/Plugins/Experimental/Compositing/Composite/Source/Composite/Public/Passes/CompositePassColorGrade.h
--- a/Engine/Plugins/Experimental/Compositing/Composite/Source/Composite/Public/Passes/CompositePassColorGrade.h
+++ b/Engine/Plugins/Experimental/Compositing/Composite/Source/Composite/Public/Passes/CompositePassColorGrade.h
@@ -53,11 +53,11 @@ public:
 
 public:
 	/** Color temperature settings. */
-	UPROPERTY(BlueprintReadWrite, EditAnywhere, Category = "Color Grading")
+	UPROPERTY(BlueprintReadWrite, EditAnywhere, Category = "")
 	FCompositeTemperatureSettings TemperatureSettings;
 
 	/** Color grade settings. */
-	UPROPERTY(BlueprintReadWrite, EditAnywhere, Category = "Color Grading", meta = (ShowPostProcessCategories))
+	UPROPERTY(BlueprintReadWrite, EditAnywhere, Category = "", meta = (ShowPostProcessCategories))
 	FColorGradingSettings ColorGradingSettings;
 };
 
diff --git a/Engine/Plugins/Experimental/Compositing/Composite/Source/CompositeEditor/Private/CompositeEditorCommands.cpp b/Engine/Plugins/Experimental/Compositing/Composite/Source/CompositeEditor/Private/CompositeEditorCommands.cpp
new file mode 100644
--- /dev/null
+++ b/Engine/Plugins/Experimental/Compositing/Composite/Source/CompositeEditor/Private/CompositeEditorCommands.cpp
@@ -0,0 +1,14 @@
+﻿// Copyright Epic Games, Inc. All Rights Reserved.
+
+#include "CompositeEditorCommands.h"
+
+#include "Framework/Commands/Commands.h"
+
+#define LOCTEXT_NAMESPACE "FCompositeEditorCommands"
+
+void FCompositeEditorCommands::RegisterCommands()
+{
+	UI_COMMAND(Enable, "Toggle Enable", "Toggle Enable State", EUserInterfaceActionType::Button, FInputChord(EModifierKey::None, EKeys::D));
+}
+
+#undef LOCTEXT_NAMESPACE
diff --git a/Engine/Plugins/Experimental/Compositing/Composite/Source/CompositeEditor/Private/CompositeEditorCommands.h b/Engine/Plugins/Experimental/Compositing/Composite/Source/CompositeEditor/Private/CompositeEditorCommands.h
new file mode 100644
--- /dev/null
+++ b/Engine/Plugins/Experimental/Compositing/Composite/Source/CompositeEditor/Private/CompositeEditorCommands.h
@@ -0,0 +1,26 @@
+﻿// Copyright Epic Games, Inc. All Rights Reserved.
+
+#pragma once
+
+#include "Framework/Commands/Commands.h"
+#include "Styling/AppStyle.h"
+
+/**
+ * 
+ */
+class FCompositeEditorCommands : public TCommands<FCompositeEditorCommands>
+{
+public:
+	FCompositeEditorCommands()
+		: TCommands<FCompositeEditorCommands>(
+			"CompositeEditor",
+			NSLOCTEXT("FCompositeEditorCommands", "CompositeEditor", "Composite Editor"),
+			NAME_None, FAppStyle::GetAppStyleSetName()
+		)
+	{ }
+
+	virtual void RegisterCommands() override;
+
+public:
+	TSharedPtr<FUICommandInfo> Enable;
+};
diff --git a/Engine/Plugins/Experimental/Compositing/Composite/Source/CompositeEditor/Private/CompositeEditorModule.cpp b/Engine/Plugins/Experimental/Compositing/Composite/Source/CompositeEditor/Private/CompositeEditorModule.cpp
--- a/Engine/Plugins/Experimental/Compositing/Composite/Source/CompositeEditor/Private/CompositeEditorModule.cpp
+++ b/Engine/Plugins/Experimental/Compositing/Composite/Source/CompositeEditor/Private/CompositeEditorModule.cpp
@@ -9,8 +9,12 @@
 #include "LevelEditor.h"
 #include "Modules/ModuleManager.h"
 #include "Widgets/Notifications/SNotificationList.h"
-
 #include "CompositeActor.h"
+#include "CompositeEditorCommands.h"
+#include "CompositeEditorStyle.h"
+#include "PropertyEditorModule.h"
+#include "Customizations/CompositeLayerPlateCustomization.h"
+#include "Layers/CompositeLayerPlate.h"
 #include "UI/SCompositeEditorPanel.h"
 
 #define LOCTEXT_NAMESPACE "CompositeEditorModule"
@@ -21,12 +25,20 @@ void FCompositeEditorModule::StartupModule()
 {
 	FCoreDelegates::OnPostEngineInit.AddRaw(this, &FCompositeEditorModule::OnPostEngineInit);
 
+	FCompositeEditorCommands::Register();
+	FCompositeEditorStyle::Register();
+	
 	RegisterTabSpawners();
+	RegisterCustomizations();
 }
 
 void FCompositeEditorModule::ShutdownModule()
 {
+	UnregisterCustomizations();
 	UnregisterTabSpawners();
+
+	FCompositeEditorCommands::Unregister();
+	FCompositeEditorStyle::Unregister();
 	
 	if (GEditor != nullptr)
 	{
@@ -97,6 +109,21 @@ void FCompositeEditorModule::UnregisterTabSpawners()
 	LevelEditorModule.OnTabManagerChanged().Remove(OnTabManagerChangedDelegateHandle);
 }
 
+void FCompositeEditorModule::RegisterCustomizations()
+{
+	FPropertyEditorModule& PropertyModule = FModuleManager::LoadModuleChecked<FPropertyEditorModule>("PropertyEditor");
+	PropertyModule.RegisterCustomClassLayout(
+		UCompositeLayerPlate::StaticClass()->GetFName(),
+		FOnGetDetailCustomizationInstance::CreateStatic(&FCompositeLayerPlateCustomization::MakeInstance)
+	);
+}
+
+void FCompositeEditorModule::UnregisterCustomizations()
+{
+	FPropertyEditorModule& PropertyModule = FModuleManager::LoadModuleChecked<FPropertyEditorModule>("PropertyEditor");
+	PropertyModule.UnregisterCustomClassLayout(UCompositeLayerPlate::StaticClass()->GetFName());
+}
+
 IMPLEMENT_MODULE(FCompositeEditorModule, CompositeEditor)
 
 #undef LOCTEXT_NAMESPACE
diff --git a/Engine/Plugins/Experimental/Compositing/Composite/Source/CompositeEditor/Private/CompositeEditorStyle.cpp b/Engine/Plugins/Experimental/Compositing/Composite/Source/CompositeEditor/Private/CompositeEditorStyle.cpp
new file mode 100644
--- /dev/null
+++ b/Engine/Plugins/Experimental/Compositing/Composite/Source/CompositeEditor/Private/CompositeEditorStyle.cpp
@@ -0,0 +1,38 @@
+﻿// Copyright Epic Games, Inc. All Rights Reserved.
+
+#include "CompositeEditorStyle.h"
+
+#include "Styling/SlateStyleMacros.h"
+#include "Styling/SlateStyleRegistry.h"
+
+FCompositeEditorStyle& FCompositeEditorStyle::Get()
+{
+	static FCompositeEditorStyle Instance;
+	return Instance;
+}
+
+void FCompositeEditorStyle::Register()
+{
+	FSlateStyleRegistry::RegisterSlateStyle(Get());
+}
+
+void FCompositeEditorStyle::Unregister()
+{
+	FSlateStyleRegistry::UnRegisterSlateStyle(Get());
+}
+
+FCompositeEditorStyle::FCompositeEditorStyle()
+	: FSlateStyleSet("CompositeEditorStyle")
+{
+	const FVector2D Icon16x16(16.0f, 16.0f);
+	const FVector2D Icon20x20(20.0f, 20.0f);
+	const FVector2D Icon64x64(64.0f, 64.0f);
+
+	FSlateStyleSet::SetContentRoot(FPaths::EngineContentDir() / TEXT("Editor/Slate/"));
+	
+	Set("CompositeEditor.Composure", new IMAGE_BRUSH_SVG("Starship/Common/ComposureCompositing", Icon16x16));
+	
+    Set("CompositeEditor.Passes.Media", new IMAGE_BRUSH_SVG("Starship/AssetIcons/MediaPlayer_16", Icon16x16));
+	Set("CompositeEditor.Passes.Scene", new IMAGE_BRUSH_SVG("Starship/Common/Levels", Icon16x16));
+	Set("CompositeEditor.Passes.Layer", new IMAGE_BRUSH_SVG("Starship/Common/Layers", Icon16x16));
+}
diff --git a/Engine/Plugins/Experimental/Compositing/Composite/Source/CompositeEditor/Private/CompositeEditorStyle.h b/Engine/Plugins/Experimental/Compositing/Composite/Source/CompositeEditor/Private/CompositeEditorStyle.h
new file mode 100644
--- /dev/null
+++ b/Engine/Plugins/Experimental/Compositing/Composite/Source/CompositeEditor/Private/CompositeEditorStyle.h
@@ -0,0 +1,16 @@
+﻿// Copyright Epic Games, Inc. All Rights Reserved.
+
+#pragma once
+#include "Styling/SlateStyle.h"
+
+class FCompositeEditorStyle	final : public FSlateStyleSet
+{
+public:
+	static FCompositeEditorStyle& Get();
+
+	static void Register();
+	static void Unregister();
+	
+private:
+	FCompositeEditorStyle();
+};
diff --git a/Engine/Plugins/Experimental/Compositing/Composite/Source/CompositeEditor/Private/Customizations/CompositeLayerPlateCustomization.cpp b/Engine/Plugins/Experimental/Compositing/Composite/Source/CompositeEditor/Private/Customizations/CompositeLayerPlateCustomization.cpp
new file mode 100644
--- /dev/null
+++ b/Engine/Plugins/Experimental/Compositing/Composite/Source/CompositeEditor/Private/Customizations/CompositeLayerPlateCustomization.cpp
@@ -0,0 +1,99 @@
+﻿// Copyright Epic Games, Inc. All Rights Reserved.
+
+#include "CompositeLayerPlateCustomization.h"
+
+#include "DetailCategoryBuilder.h"
+#include "DetailLayoutBuilder.h"
+#include "DetailWidgetRow.h"
+#include "IDetailGroup.h"
+#include "IPropertyUtilities.h"
+#include "Layers/CompositeLayerPlate.h"
+#include "UI/SCompositePlatePassPanel.h"
+#include "Widgets/Text/STextBlock.h"
+
+#define LOCTEXT_NAMESPACE "FCompositeLayerPlateCustomization"
+
+/** Temporary cvar to ensure that the old passes lists can still be exposed in case something is missing or incorrect with the new pass table */
+TAutoConsoleVariable<bool> CVarDisablePassTableCustomization(
+	TEXT("Composure.DisablePassTableCustomization"),
+	/*DefaultValue*/false,
+	TEXT("If `true`, composite plates do not show the new pass table in their details panel, instead showing the default passes lists")
+);
+
+TSharedRef<IDetailCustomization> FCompositeLayerPlateCustomization::MakeInstance()
+{
+	return MakeShared<FCompositeLayerPlateCustomization>();
+}
+
+void FCompositeLayerPlateCustomization::CustomizeDetails(const TSharedPtr<IDetailLayoutBuilder>& DetailBuilder)
+{
+	CachedDetailBuilder = DetailBuilder;
+	IDetailCustomization::CustomizeDetails(DetailBuilder);
+}
+
+void FCompositeLayerPlateCustomization::CustomizeDetails(IDetailLayoutBuilder& DetailLayout)
+{
+	// TODO: Once we are confident in the pass table UI, this can be removed, but for now allow fallback to non-customized details panel for testing/debug
+	if (CVarDisablePassTableCustomization.GetValueOnAnyThread())
+	{
+		return;
+	}
+	
+	TArray<TWeakObjectPtr<UCompositeLayerPlate>> Objects = DetailLayout.GetObjectsOfTypeBeingCustomized<UCompositeLayerPlate>();
+
+	DetailLayout.HideProperty(GET_MEMBER_NAME_CHECKED(UCompositeLayerPlate, MediaPasses));
+	DetailLayout.HideProperty(GET_MEMBER_NAME_CHECKED(UCompositeLayerPlate, LayerPasses));
+	
+	IDetailCategoryBuilder& PlateCategory = DetailLayout.EditCategoryAllowNone(NAME_None);
+
+	// Add all default simple properties ahead of the custom group
+	TArray<TSharedRef<IPropertyHandle>> PlateProperties;
+	PlateCategory.GetDefaultProperties(PlateProperties, /* bSimpleProperties */ true, /* bAdvancedProperties */ false);
+	for (const TSharedRef<IPropertyHandle>& Property : PlateProperties)
+	{
+		if (!Property->IsCustomized())
+		{
+			PlateCategory.AddProperty(Property);
+		}
+	}
+	
+	if (Objects.Num() == 1 && Objects[0].IsValid())
+	{
+		UCompositeLayerPlate* Plate = Cast<UCompositeLayerPlate>(Objects[0].Get());
+		
+		IDetailGroup& PassesGroup = PlateCategory.AddGroup("Passes", LOCTEXT("PassesGroupName", "Passes"), false, true);
+
+		PassesGroup.AddWidgetRow()
+		.WholeRowContent()
+		[
+			SNew(SCompositePlatePassPanel, Plate)
+			.OnLayoutSizeChanged_Lambda([this]
+			{
+				TSharedPtr<IDetailLayoutBuilder> PinnedDetailBuilder = CachedDetailBuilder.Pin();
+				if (PinnedDetailBuilder.IsValid())
+				{
+					PinnedDetailBuilder->GetPropertyUtilities()->RequestRefresh();
+				}
+			})
+		];
+	}  
+	else
+	{
+		// Can't display plate pass panel if multiple plates are selected, so simply put a "Multiple Values" entry in the property list
+		PlateCategory.AddCustomRow(LOCTEXT("PassesGroupName", "Passes"))
+		.NameContent()
+		[
+			SNew(STextBlock)
+			.Text(LOCTEXT("PassesGroupName", "Passes"))
+			.Font(IDetailLayoutBuilder::GetDetailFont())
+		]
+		.ValueContent()
+		[
+			SNew(STextBlock)
+			.Text(LOCTEXT("MultipleValues", "Multiple Values"))
+			.Font(IDetailLayoutBuilder::GetDetailFont())
+		];
+	}
+}
+
+#undef LOCTEXT_NAMESPACE
diff --git a/Engine/Plugins/Experimental/Compositing/Composite/Source/CompositeEditor/Private/Customizations/CompositeLayerPlateCustomization.h b/Engine/Plugins/Experimental/Compositing/Composite/Source/CompositeEditor/Private/Customizations/CompositeLayerPlateCustomization.h
new file mode 100644
--- /dev/null
+++ b/Engine/Plugins/Experimental/Compositing/Composite/Source/CompositeEditor/Private/Customizations/CompositeLayerPlateCustomization.h
@@ -0,0 +1,25 @@
+﻿// Copyright Epic Games, Inc. All Rights Reserved.
+
+#pragma once
+
+#include "IDetailCustomization.h"
+
+class SCompositePlatePassPanel;
+
+/**
+ * Customization for the composite plate layer, primary for displaying a custom widget for passes and composite meshes
+ */
+class FCompositeLayerPlateCustomization : public IDetailCustomization
+{
+public:
+	/** Makes a new instance of the details customization */
+	static TSharedRef<IDetailCustomization> MakeInstance();
+
+	// IDetailCustomization interface
+	virtual void CustomizeDetails(const TSharedPtr<IDetailLayoutBuilder>& DetailBuilder) override;
+	virtual void CustomizeDetails(IDetailLayoutBuilder& DetailLayout) override;
+	// End of IDetailCustomization interface
+
+private:
+	TWeakPtr<IDetailLayoutBuilder> CachedDetailBuilder;
+};
diff --git a/Engine/Plugins/Experimental/Compositing/Composite/Source/CompositeEditor/Private/UI/SCompositeEditorPanel.cpp b/Engine/Plugins/Experimental/Compositing/Composite/Source/CompositeEditor/Private/UI/SCompositeEditorPanel.cpp
--- a/Engine/Plugins/Experimental/Compositing/Composite/Source/CompositeEditor/Private/UI/SCompositeEditorPanel.cpp
+++ b/Engine/Plugins/Experimental/Compositing/Composite/Source/CompositeEditor/Private/UI/SCompositeEditorPanel.cpp
@@ -3,6 +3,7 @@
 #include "SCompositeEditorPanel.h"
 
 #include "CompositeActor.h"
+#include "CompositeEditorStyle.h"
 #include "DetailCategoryBuilder.h"
 #include "DetailLayoutBuilder.h"
 #include "DetailWidgetRow.h"
@@ -393,8 +394,8 @@ void SCompositeEditorPanel::RegisterTabSpawner()
 		LevelEditorTabManager->RegisterTabSpawner(CompositingEditorPanel::CompositeEditorTabName, FOnSpawnTab::CreateStatic(&CompositingEditorPanel::CreateTab))
 			.SetDisplayName(LOCTEXT("PanelDisplayName", "Composure"))
 			.SetTooltipText(LOCTEXT("PanelTooltipText", "Panel for viewing and editing compositing actors in the current level"))
-			.SetGroup(WorkspaceMenu::GetMenuStructure().GetLevelEditorVirtualProductionCategory());
-			//.SetIcon(FSlateIcon(FVPUtilitiesEditorStyle::GetStyleSetName(), "TabIcons.Genlock.Small")); // TODO: find the correct icon
+			.SetGroup(WorkspaceMenu::GetMenuStructure().GetLevelEditorVirtualProductionCategory())
+			.SetIcon(FSlateIcon(FCompositeEditorStyle::Get().GetStyleSetName(), "CompositeEditor.Composure"));
 	}
 }
 
diff --git a/Engine/Plugins/Experimental/Compositing/Composite/Source/CompositeEditor/Private/UI/SCompositePanelLayerTree.cpp b/Engine/Plugins/Experimental/Compositing/Composite/Source/CompositeEditor/Private/UI/SCompositePanelLayerTree.cpp
--- a/Engine/Plugins/Experimental/Compositing/Composite/Source/CompositeEditor/Private/UI/SCompositePanelLayerTree.cpp
+++ b/Engine/Plugins/Experimental/Compositing/Composite/Source/CompositeEditor/Private/UI/SCompositePanelLayerTree.cpp
@@ -2,8 +2,8 @@
 
 #include "SCompositePanelLayerTree.h"
 
-#include "AudioMixerBlueprintLibrary.h"
 #include "CompositeActor.h"
+#include "CompositeEditorCommands.h"
 #include "Editor.h"
 #include "EngineUtils.h"
 #include "Factories.h"
@@ -23,7 +23,6 @@
 #include "Widgets/Text/SInlineEditableTextBlock.h"
 #include "Widgets/Views/STableRow.h"
 #include "Widgets/Views/STreeView.h"
-#include "Windows/WindowsPlatformApplicationMisc.h"
 
 #define LOCTEXT_NAMESPACE "SCompositePanelLayerTree"
 
@@ -403,9 +402,9 @@ private:
 		return SNullWidget::NullWidget;
 	}
 
-	ACompositeActor* GetCompositeActor() const
+	TStrongObjectPtr<ACompositeActor> GetCompositeActor() const
 	{
-		return TreeItem->HasValidCompositeActor() ? TreeItem->CompositeActor.Get() : nullptr;
+		return TreeItem->HasValidCompositeActor() ? TreeItem->CompositeActor.Pin() : nullptr;
 	}
 
 	/** Gets the checkbox state for the row item's Enabled flag */
@@ -422,7 +421,7 @@ private:
 			return ECheckBoxState::Unchecked;
 		}
 
-		if (ACompositeActor* CompositeActor = GetCompositeActor())
+		if (TStrongObjectPtr<ACompositeActor> CompositeActor = GetCompositeActor())
 		{
 			return CompositeActor->IsEnabled() ? ECheckBoxState::Checked : ECheckBoxState::Unchecked;
 		}
@@ -449,7 +448,7 @@ private:
 			return;
 		}
 		
-		if (ACompositeActor* CompositeActor = GetCompositeActor())
+		if (TStrongObjectPtr<ACompositeActor> CompositeActor = GetCompositeActor())
 		{
 			FScopedTransaction SetEnabledTransaction(LOCTEXT("SetEnabledTransaction", "Set Enabled"));
 			CompositeActor->Modify();
@@ -520,7 +519,7 @@ private:
 			return nullptr;
 		}
 		
-		if (ACompositeActor* CompositeActor = GetCompositeActor())
+		if (TStrongObjectPtr<ACompositeActor> CompositeActor = GetCompositeActor())
 		{
 			return FSlateIconFinder::FindIconForClass(CompositeActor->GetClass()).GetIcon(); 
 		}
@@ -548,7 +547,7 @@ private:
 			return LOCTEXT("LayerNotConfiguredLabel", "Layer not configured");
 		}
 
-		if (ACompositeActor* CompositeActor = GetCompositeActor())
+		if (TStrongObjectPtr<ACompositeActor> CompositeActor = GetCompositeActor())
 		{
 			return FText::FromString(CompositeActor->GetActorNameOrLabel());
 		}
@@ -574,10 +573,10 @@ private:
 			return;
 		}
 
-		if (ACompositeActor* CompositeActor = GetCompositeActor())
+		if (TStrongObjectPtr<ACompositeActor> CompositeActor = GetCompositeActor())
 		{
 			FScopedTransaction RenameActorTransaction(LOCTEXT("RenameActorTransaction", "Rename Composite Actor"));
-			FActorLabelUtilities::RenameExistingActor(CompositeActor, InNewName.ToString(), true);
+			FActorLabelUtilities::RenameExistingActor(CompositeActor.Get(), InNewName.ToString(), true);
 		}
 	}
 
@@ -594,7 +593,7 @@ private:
 			return LOCTEXT("NoneTypeLabel", "None");
 		}
 
-		if (ACompositeActor* CompositeActor = GetCompositeActor())
+		if (TStrongObjectPtr<ACompositeActor> CompositeActor = GetCompositeActor())
 		{
 			return CompositeActor->GetClass()->GetDisplayNameText();
 		}
@@ -879,7 +878,8 @@ FReply SCompositePanelLayerTree::OnKeyDown(const FGeometry& MyGeometry, const FK
 
 void SCompositePanelLayerTree::PostUndo(bool bSuccess)
 {
-	FillCompositeActorTree();
+	constexpr bool bPreserveSelection = true;
+	FillCompositeActorTree(bPreserveSelection);
 }
 
 void SCompositePanelLayerTree::PostRedo(bool bSuccess)
@@ -912,10 +912,27 @@ void SCompositePanelLayerTree::BindCommands()
 	CommandList->MapAction(FGenericCommands::Get().Rename,
 		FUIAction(FExecuteAction::CreateSP(this, &SCompositePanelLayerTree::RenameSelectedItem),
 		FCanExecuteAction::CreateSP(this, &SCompositePanelLayerTree::CanRenameSelectedItem)));
+
+	CommandList->MapAction(FCompositeEditorCommands::Get().Enable,
+		FExecuteAction::CreateSP(this, &SCompositePanelLayerTree::EnableSelectedItems),
+		FCanExecuteAction::CreateSP(this, &SCompositePanelLayerTree::CanEnableSelectedItems));
 }
 
-void SCompositePanelLayerTree::FillCompositeActorTree()
+void SCompositePanelLayerTree::FillCompositeActorTree(bool bPreserveSelection)
 {
+	TArray<TPair<ACompositeActor*, int32>> ItemsToReselect;
+	if (bPreserveSelection)
+	{
+		TArray<FCompositeActorTreeItemPtr> SelectedItems = TreeView.IsValid() ? TreeView->GetSelectedItems() : TArray<FCompositeActorTreeItemPtr>();
+		for (const FCompositeActorTreeItemPtr& Item : SelectedItems)
+		{
+			if (Item->CompositeActor.IsValid())
+			{
+				ItemsToReselect.Add(TPair<ACompositeActor*, int32>(Item->CompositeActor.Get(), Item->LayerIndex));
+			}
+		}
+	}
+	
 	CompositeActorTreeItems.Empty();
 	FilteredCompositeActorTreeItems.Empty();
 	
@@ -954,6 +971,35 @@ void SCompositePanelLayerTree::FillCompositeActorTree()
 	}
 
 	RefreshAndExpandTreeView();
+
+	if (TreeView.IsValid())
+	{
+		for (const TPair<ACompositeActor*, int32>& Pair : ItemsToReselect)
+		{
+			ACompositeActor* CompositeActorToReselect = Pair.Key;
+			int32 LayerToReselect = Pair.Value;
+
+			const int32 CompositeActorItemIndex = FilteredCompositeActorTreeItems.IndexOfByPredicate([CompositeActorToReselect](const FCompositeActorTreeItemPtr& InItem)
+			{
+				return InItem->CompositeActor == CompositeActorToReselect;
+			});
+
+			if (CompositeActorItemIndex == INDEX_NONE)
+			{
+				continue;
+			}
+
+			FCompositeActorTreeItemPtr CompositeActorItem = FilteredCompositeActorTreeItems[CompositeActorItemIndex];
+			if (LayerToReselect == CompositeActorItem->LayerIndex)
+			{
+				TreeView->SetItemSelection(CompositeActorItem, true);
+			}
+			else if (CompositeActorItem->Children.IsValidIndex(LayerToReselect) && !CompositeActorItem->Children[LayerToReselect]->bFilteredOut)
+			{
+				TreeView->SetItemSelection(CompositeActorItem->Children[LayerToReselect], true);
+			}
+		}
+	}
 }
 
 void SCompositePanelLayerTree::FillCompositeActorTreeItem(FCompositeActorTreeItemPtr& InOutTreeItem)
@@ -965,12 +1011,12 @@ void SCompositePanelLayerTree::FillCompositeActorTreeItem(FCompositeActorTreeIte
 	
 	InOutTreeItem->Children.Empty();
 
-	if (ACompositeActor* CompositeActor = InOutTreeItem->CompositeActor.Get())
+	if (TStrongObjectPtr<ACompositeActor> CompositeActor = InOutTreeItem->CompositeActor.Pin())
 	{
 		for (int32 Index = 0; Index < CompositeActor->CompositeLayers.Num(); ++Index)
 		{
 			FCompositeActorTreeItemPtr NewLayerTreeItem = MakeShared<FCompositeActorTreeItem>();
-			NewLayerTreeItem->CompositeActor = CompositeActor;
+			NewLayerTreeItem->CompositeActor = CompositeActor.Get();
 			NewLayerTreeItem->LayerIndex = Index;
 			NewLayerTreeItem->bFilteredOut = Toolbar.IsValid() ? !Toolbar->ItemPassesFilters(NewLayerTreeItem) : false;
 			
@@ -1102,6 +1148,7 @@ void SCompositePanelLayerTree::OnPilotCameraToggled(bool bPilotCamera)
 		if (bPilotCamera)
 		{
 			LevelEditorSubsystem->PilotLevelActor(CompositeCameraActor);
+			LevelEditorSubsystem->SetExactCameraView(true);
 		}
 		else
 		{
@@ -1340,6 +1387,10 @@ TSharedPtr<SWidget> SCompositePanelLayerTree::CreateTreeContextMenu()
 		MenuBuilder.AddMenuEntry(FGenericCommands::Get().Duplicate);
 		MenuBuilder.AddMenuEntry(FGenericCommands::Get().Delete);
 		MenuBuilder.AddMenuEntry(FGenericCommands::Get().Rename);
+
+		MenuBuilder.AddSeparator();
+
+		MenuBuilder.AddMenuEntry(FCompositeEditorCommands::Get().Enable);
 		
 		return MenuBuilder.MakeWidget();
 	}
@@ -1355,8 +1406,8 @@ void SCompositePanelLayerTree::OnLayerAdded(const UClass* InLayerClass)
 		return;
 	}
 
-	ACompositeActor* CompositeActor = SelectedItems[0]->CompositeActor.Get();
-	if (!CompositeActor)
+	TStrongObjectPtr<ACompositeActor> CompositeActor = SelectedItems[0]->CompositeActor.Pin();
+	if (!CompositeActor.IsValid())
 	{
 		return;
 	}
@@ -1365,12 +1416,23 @@ void SCompositePanelLayerTree::OnLayerAdded(const UClass* InLayerClass)
 	FScopedTransaction LayerAddedTransaction(LOCTEXT("LayerAddedTransaction", "Add Layer"));
 	CompositeActor->Modify();
 
-	UCompositeLayerBase* NewLayer = NewObject<UCompositeLayerBase>(CompositeActor, InLayerClass);
+	UCompositeLayerBase* NewLayer = NewObject<UCompositeLayerBase>(CompositeActor.Get(), InLayerClass);
 	NewLayer->Name = NewLayer->GetName();
+
+	FProperty* CompositeLayersProperty = FindFProperty<FProperty>(ACompositeActor::StaticClass(), GET_MEMBER_NAME_CHECKED(ACompositeActor, CompositeLayers));
+	CompositeActor->PreEditChange(CompositeLayersProperty);
 	
 	CompositeActor->CompositeLayers.Insert(NewLayer, NewLayerIndex);
 
-	RefreshCompositeActorLayers(CompositeActor);
+	FEditPropertyChain EditChain;
+	EditChain.AddTail(CompositeLayersProperty);
+	EditChain.SetActivePropertyNode(CompositeLayersProperty);
+
+	FPropertyChangedEvent ChangeEvent(CompositeLayersProperty, EPropertyChangeType::ArrayAdd);
+	FPropertyChangedChainEvent ChainEvent(EditChain, ChangeEvent);
+	CompositeActor->PostEditChangeChainProperty(ChainEvent);
+	
+	RefreshCompositeActorLayers(CompositeActor.Get());
 }
 
 FReply SCompositePanelLayerTree::OnLayerSoloToggled(const FCompositeActorTreeItemPtr& InLayerTreeItem)
@@ -1381,7 +1443,12 @@ FReply SCompositePanelLayerTree::OnLayerSoloToggled(const FCompositeActorTreeIte
 		return FReply::Unhandled();
 	}
 
-	ACompositeActor* CompositeActor = InLayerTreeItem->CompositeActor.Get();
+	TStrongObjectPtr<ACompositeActor> CompositeActor = InLayerTreeItem->CompositeActor.Pin();
+	if (!CompositeActor.IsValid())
+	{
+		return FReply::Unhandled();
+	}
+	
 	UCompositeLayerBase* CompositeLayer = InLayerTreeItem->GetCompositeLayer();
 
 	if (!CompositeLayer)
@@ -1419,7 +1486,12 @@ void SCompositePanelLayerTree::OnLayerMoved(const FCompositeActorTreeItemPtr& In
 		return;
 	}
 	
-	ACompositeActor* CompositeActor = InTreeItem->CompositeActor.Get();
+	TStrongObjectPtr<ACompositeActor> CompositeActor = InTreeItem->CompositeActor.Pin();
+	if (!CompositeActor.IsValid())
+	{
+		return;
+	}
+	
 	UCompositeLayerBase* CompositeLayer = InTreeItem->GetCompositeLayer();
 	
 	if (!CompositeLayer || !CompositeActor->CompositeLayers.IsValidIndex(InDestIndex))
@@ -1429,11 +1501,22 @@ void SCompositePanelLayerTree::OnLayerMoved(const FCompositeActorTreeItemPtr& In
 	
 	FScopedTransaction MoveLayerTransaction(LOCTEXT("MoveLayerTransaction", "Move Layer"));
 	CompositeActor->Modify();
+
+	FProperty* CompositeLayersProperty = FindFProperty<FProperty>(ACompositeActor::StaticClass(), GET_MEMBER_NAME_CHECKED(ACompositeActor, CompositeLayers));
+	CompositeActor->PreEditChange(CompositeLayersProperty);
 	
 	CompositeActor->CompositeLayers.RemoveAt(InTreeItem->LayerIndex);
 	CompositeActor->CompositeLayers.Insert(CompositeLayer, InDestIndex);
+
+	FEditPropertyChain EditChain;
+	EditChain.AddTail(CompositeLayersProperty);
+	EditChain.SetActivePropertyNode(CompositeLayersProperty);
+
+	FPropertyChangedEvent ChangeEvent(CompositeLayersProperty, EPropertyChangeType::ArrayMove);
+	FPropertyChangedChainEvent ChainEvent(EditChain, ChangeEvent);
+	CompositeActor->PostEditChangeChainProperty(ChainEvent);
 	
-	RefreshCompositeActorLayers(CompositeActor);
+	RefreshCompositeActorLayers(CompositeActor.Get());
 }
 
 void SCompositePanelLayerTree::OnActorAddedToLevel(AActor* Actor)
@@ -1641,7 +1724,20 @@ void SCompositePanelLayerTree::PasteSelectedItems()
 		for (UCompositeLayerBase* LayerToPaste : Factory.CompositeLayers)
 		{
 			UCompositeLayerBase* NewLayer = DuplicateObject(LayerToPaste, CompositeActor);
+			
+			FProperty* CompositeLayersProperty = FindFProperty<FProperty>(ACompositeActor::StaticClass(), GET_MEMBER_NAME_CHECKED(ACompositeActor, CompositeLayers));
+			CompositeActor->PreEditChange(CompositeLayersProperty);
+			
 			CompositeActor->CompositeLayers.Insert(NewLayer, IndexToPasteAt);
+
+			FEditPropertyChain EditChain;
+			EditChain.AddTail(CompositeLayersProperty);
+			EditChain.SetActivePropertyNode(CompositeLayersProperty);
+
+			FPropertyChangedEvent ChangeEvent(CompositeLayersProperty, EPropertyChangeType::ArrayAdd);
+			FPropertyChangedChainEvent ChainEvent(EditChain, ChangeEvent);
+			CompositeActor->PostEditChangeChainProperty(ChainEvent);
+			
 			LayersToSelect[CompositeActor].Add(IndexToPasteAt);
 			++IndexToPasteAt;
 		}
@@ -1734,12 +1830,23 @@ void SCompositePanelLayerTree::DeleteSelectedItems()
 			// Delete the layers in reverse order so that indices don't get messed up as we delete
 			TArray<int32>& Layers = CompositeActorLayerList.Value;
 			Layers.Sort();
+
+			FProperty* CompositeLayersProperty = FindFProperty<FProperty>(ACompositeActor::StaticClass(), GET_MEMBER_NAME_CHECKED(ACompositeActor, CompositeLayers));
+			CompositeActor->PreEditChange(CompositeLayersProperty);
 			
 			for (int32 Index = Layers.Num() - 1; Index >= 0; --Index)
 			{
 				CompositeActor->CompositeLayers.RemoveAt(Layers[Index]);
 			}
 
+			FEditPropertyChain EditChain;
+			EditChain.AddTail(CompositeLayersProperty);
+			EditChain.SetActivePropertyNode(CompositeLayersProperty);
+
+			FPropertyChangedEvent ChangeEvent(CompositeLayersProperty, EPropertyChangeType::ArrayAdd);
+			FPropertyChangedChainEvent ChainEvent(EditChain, ChangeEvent);
+			CompositeActor->PostEditChangeChainProperty(ChainEvent);
+			
 			int32 TreeItemIndex = CompositeActorTreeItems.IndexOfByPredicate([CompositeActor](const FCompositeActorTreeItemPtr& InTreeItem)
 			{
 				return InTreeItem->CompositeActor.Get() == CompositeActor;
@@ -1792,4 +1899,69 @@ bool SCompositePanelLayerTree::CanRenameSelectedItem()
 	return false;
 }
 
+void SCompositePanelLayerTree::EnableSelectedItems()
+{
+	TArray<FCompositeActorTreeItemPtr> SelectedItems = TreeView.IsValid() ? TreeView->GetSelectedItems() : TArray<FCompositeActorTreeItemPtr>();
+
+	bool bAnyEnabled = false;
+	for (const FCompositeActorTreeItemPtr& Item : SelectedItems)
+	{
+		TStrongObjectPtr<ACompositeActor> PinnedCompositeActor = Item->CompositeActor.Pin();
+		if (!PinnedCompositeActor.IsValid())
+		{
+			continue;
+		}
+		
+		if (Item->HasValidCompositeLayer())
+		{
+			if (UCompositeLayerBase* Layer = Item->GetCompositeLayer())
+			{
+				if (Layer->bIsEnabled)
+				{
+					bAnyEnabled = true;
+				}
+			}
+		}
+		else
+		{
+			if (PinnedCompositeActor->IsEnabled())
+			{
+				bAnyEnabled = true;
+			}
+		}
+	}
+
+	const bool bSetEnabled = !bAnyEnabled;
+
+	FScopedTransaction SetEnabledTransaction(LOCTEXT("SetEnabledTransaction", "Set Enabled"));
+	for (const FCompositeActorTreeItemPtr& Item : SelectedItems)
+	{
+		TStrongObjectPtr<ACompositeActor> PinnedCompositeActor = Item->CompositeActor.Pin();
+		if (!PinnedCompositeActor.IsValid())
+		{
+			continue;
+		}
+		
+		if (Item->HasValidCompositeLayer())
+		{
+			if (UCompositeLayerBase* Layer = Item->GetCompositeLayer())
+			{
+				Layer->Modify();
+				Layer->bIsEnabled = bSetEnabled;
+			}
+		}
+		else
+		{
+			PinnedCompositeActor->Modify();
+			PinnedCompositeActor->SetEnabled(bSetEnabled);
+		}
+	}
+}
+
+bool SCompositePanelLayerTree::CanEnableSelectedItems()
+{
+	TArray<FCompositeActorTreeItemPtr> SelectedItems = TreeView.IsValid() ? TreeView->GetSelectedItems() : TArray<FCompositeActorTreeItemPtr>();
+	return !SelectedItems.IsEmpty();
+}
+
 #undef LOCTEXT_NAMESPACE
diff --git a/Engine/Plugins/Experimental/Compositing/Composite/Source/CompositeEditor/Private/UI/SCompositePanelLayerTree.h b/Engine/Plugins/Experimental/Compositing/Composite/Source/CompositeEditor/Private/UI/SCompositePanelLayerTree.h
--- a/Engine/Plugins/Experimental/Compositing/Composite/Source/CompositeEditor/Private/UI/SCompositePanelLayerTree.h
+++ b/Engine/Plugins/Experimental/Compositing/Composite/Source/CompositeEditor/Private/UI/SCompositePanelLayerTree.h
@@ -70,7 +70,7 @@ private:
 	void BindCommands();
 	
 	/** Fills the tree view with all composite actors found on the current level */
-	void FillCompositeActorTree();
+	void FillCompositeActorTree(bool bPreserveSelection = false);
 
 	/** Fills the children of the composite actor tree item with tree items for each of the composite actor's layers */
 	void FillCompositeActorTreeItem(FCompositeActorTreeItemPtr& InOutTreeItem);
@@ -147,6 +147,8 @@ private:
 	bool CanDeleteSelectedItems();
 	void RenameSelectedItem();
 	bool CanRenameSelectedItem();
+	void EnableSelectedItems();
+	bool CanEnableSelectedItems();
 	
 private:
 	/** Tree view widget that displays the level's composite actors and their layers */
diff --git a/Engine/Plugins/Experimental/Compositing/Composite/Source/CompositeEditor/Private/UI/SCompositePassTree.cpp b/Engine/Plugins/Experimental/Compositing/Composite/Source/CompositeEditor/Private/UI/SCompositePassTree.cpp
new file mode 100644
--- /dev/null
+++ b/Engine/Plugins/Experimental/Compositing/Composite/Source/CompositeEditor/Private/UI/SCompositePassTree.cpp
@@ -0,0 +1,1752 @@
+﻿// Copyright Epic Games, Inc. All Rights Reserved.
+
+
+#include "SCompositePassTree.h"
+
+#include "CompositeEditorCommands.h"
+#include "CompositeEditorStyle.h"
+#include "Factories.h"
+#include "SActionButton.h"
+#include "ScopedTransaction.h"
+#include "UnrealExporter.h"
+#include "DragAndDrop/DecoratedDragDropOp.h"
+#include "Exporters/Exporter.h"
+#include "Filters/GenericFilter.h"
+#include "Filters/SBasicFilterBar.h"
+#include "Framework/Commands/GenericCommands.h"
+#include "HAL/PlatformApplicationMisc.h"
+#include "Layers/CompositeLayerPlate.h"
+#include "Misc/StringOutputDevice.h"
+#include "Styling/SlateIconFinder.h"
+#include "UObject/Package.h"
+#include "Widgets/Text/SInlineEditableTextBlock.h"
+#include "Widgets/Views/STableRow.h"
+#include "Widgets/Views/STreeView.h"
+
+#define LOCTEXT_NAMESPACE "SCompositePassTree"
+
+namespace CompositePassTree
+{
+	static const FName PassTreeColumn_Enabled = "Enabled";
+	static const FName PassTreeColumn_Name = "Name";
+}
+
+/**
+ * Widget for the filter bar, needs to be a subclass that overrides MakeAddFilterMenu to give the filter bar its own unique menu name,
+ * otherwise the editor will get confused with any other basic filter bar used elsewhere.
+ */
+template<typename TFilterType>
+class SCompositePassesFilterBar : public SBasicFilterBar<TFilterType>
+{
+	using Super = SBasicFilterBar<TFilterType>;
+	
+public:
+
+	using FOnFilterChanged = typename SBasicFilterBar<TFilterType>::FOnFilterChanged;
+	using FCreateTextFilter = typename SBasicFilterBar<TFilterType>::FCreateTextFilter;
+
+	SLATE_BEGIN_ARGS(SCompositePassesFilterBar)
+	{}
+		SLATE_EVENT(SCompositePassesFilterBar<TFilterType>::FOnFilterChanged, OnFilterChanged)
+		SLATE_ARGUMENT(TArray<TSharedRef<FFilterBase<TFilterType>>>, CustomFilters)	
+	SLATE_END_ARGS()
+
+	void Construct(const FArguments& InArgs)
+	{
+		typename SBasicFilterBar<TFilterType>::FArguments Args;
+		Args._OnFilterChanged = InArgs._OnFilterChanged;
+		Args._CustomFilters = InArgs._CustomFilters;
+		Args._UseSectionsForCategories = true;
+		Args._bPinAllFrontendFilters = true;
+		
+		SBasicFilterBar<TFilterType>::Construct(Args.FilterPillStyle(EFilterPillStyle::Basic));
+	}
+
+private:
+	virtual TSharedRef<SWidget> MakeAddFilterMenu() override
+	{
+		const FName FilterMenuName = "CompositePlatePassesFilterBar.FilterMenu";
+		if (!UToolMenus::Get()->IsMenuRegistered(FilterMenuName))
+		{
+			UToolMenu* Menu = UToolMenus::Get()->RegisterMenu(FilterMenuName);
+			Menu->bShouldCloseWindowAfterMenuSelection = true;
+			Menu->bCloseSelfOnly = true;
+
+			Menu->AddDynamicSection(NAME_None, FNewToolMenuDelegate::CreateLambda([](UToolMenu* InMenu)
+			{
+				if (UFilterBarContext* Context = InMenu->FindContext<UFilterBarContext>())
+				{
+					Context->PopulateFilterMenu.ExecuteIfBound(InMenu);
+					Context->OnExtendAddFilterMenu.ExecuteIfBound(InMenu);
+				}
+			}));
+		}
+
+		UFilterBarContext* FilterBarContext = NewObject<UFilterBarContext>();
+		FilterBarContext->PopulateFilterMenu = FOnPopulateAddFilterMenu::CreateLambda([this](UToolMenu* Menu)
+		{
+			Super::PopulateCommonFilterSections(Menu);
+			Super::PopulateCustomFilters(Menu);
+		});
+		FToolMenuContext ToolMenuContext(FilterBarContext);
+		
+		return UToolMenus::Get()->GenerateWidget(FilterMenuName, ToolMenuContext);
+	}
+};
+
+/** Custom filter for the filter bar so that events can be received for when the filter is activated or deactivated */
+class FCompositePassTreeFilter : public FGenericFilter<SCompositePassTree::FPassTreeItemPtr>
+{
+public:
+	DECLARE_DELEGATE_OneParam(FOnActiveStateChanged, bool)
+	
+public:
+	FCompositePassTreeFilter(const TSharedPtr<FFilterCategory>& InCategory,
+		const FString& InName,
+		const FText& DisplayName,
+		const FOnItemFiltered& InFilterDelegate)
+		: FGenericFilter<TSharedPtr<SCompositePassTree::FPassTreeItem>>(InCategory, InName, DisplayName, InFilterDelegate)
+	{ }
+
+	/** Activates the filter directly, avoiding invoking OnActiveStateChanged */
+	void ActivateFilter()
+	{
+		TGuardValue<bool> Guard(bStateChangeGuard, true);
+		SetActive(true);
+	}
+
+	/** Deactivates the filter directly, avoiding invoking OnActiveStateChanged */
+	void DeactivateFilter()
+	{
+		TGuardValue<bool> Guard(bStateChangeGuard, true);
+		SetActive(false);
+	}
+
+	/** Raised when the state of the filter is chanaged */
+	virtual void ActiveStateChanged(bool bActive) override
+	{
+		if (bStateChangeGuard)
+		{
+			return;
+		}
+
+		OnActiveStateChanged.ExecuteIfBound(bActive);
+	}
+
+public:
+	FOnActiveStateChanged OnActiveStateChanged;
+	bool bStateChangeGuard = false;
+};
+
+/** Toolbar widget that contains filtering and an add button */
+class SCompositePassTreeToolbar : public SCompoundWidget
+{
+private:
+	using FFilterType = SCompositePassTree::FPassTreeItemPtr;
+	using FTreeFilter =  TTextFilter<FFilterType>;
+	
+public:
+	DECLARE_DELEGATE(FOnFilterChanged)
+	DECLARE_DELEGATE_OneParam(FOnPassAdded, const UClass* /* InPassClass */)
+	
+	SLATE_BEGIN_ARGS(SCompositePassTreeToolbar) {}
+		SLATE_EVENT(FOnFilterChanged, OnFilterChanged)
+		SLATE_EVENT(FOnPassAdded, OnPassAdded)
+	SLATE_END_ARGS()
+
+	void Construct(const FArguments& InArgs)
+	{
+		OnFilterChanged = InArgs._OnFilterChanged;
+		OnPassAdded = InArgs._OnPassAdded;
+
+		TArray<TSharedRef<FFilterBase<FFilterType>>> AllFilters;
+		InitializeFilters(AllFilters);
+		
+		FilterBar = SNew(SCompositePassesFilterBar<FFilterType>)
+			.CustomFilters(AllFilters)
+			.OnFilterChanged(this, &SCompositePassTreeToolbar::FilterChanged);
+
+		AllPassesFilter->SetActive(true);
+		
+		ChildSlot
+		[
+			SNew(SVerticalBox)
+
+			+SVerticalBox::Slot()
+			.AutoHeight()
+			[
+				SNew(SHorizontalBox)
+
+				+SHorizontalBox::Slot()
+				.AutoWidth()
+				.Padding(0.0f, 0.0f, 4.0f, 0.0f)
+				[
+					SNew(SActionButton)
+					.ActionButtonType(EActionButtonType::Simple)
+					.Icon(FAppStyle::GetBrush("Icons.PlusCircle"))
+					.ToolTipText(LOCTEXT("AddPassButtonToolTip", "Add a new pass to the selected list of passes"))
+					.OnGetMenuContent(this, &SCompositePassTreeToolbar::GetAddPassMenuContent)
+				]
+				
+				+SHorizontalBox::Slot()
+				.FillWidth(1.0f)
+				[
+					SAssignNew(TextFilterSearchBox, SFilterSearchBox)
+					.HintText(LOCTEXT("FilterSearch", "Search..."))
+					.ToolTipText(LOCTEXT("FilterSearchHint", "Type here to search for specific passes"))
+					.OnTextChanged_Lambda([this](const FText& InText)
+					{
+						TreeTextFilter->SetRawFilterText(InText);
+					})
+				]
+			]
+
+			+SVerticalBox::Slot()
+			.AutoHeight()
+			[
+				FilterBar.ToSharedRef()
+			]
+		];
+	}
+
+	bool ItemPassesFilters(const FFilterType& InItem) const
+	{
+		// Performs an OR check on any of the filters that are enabled in the filter bar, mirroring the scene outliner filter bar behavior
+		auto PassesAnyFilterBarFilter = [this](const FFilterType& InItem)
+		{
+			TSharedPtr<TFilterCollection<FFilterType>> FilterCollection = FilterBar->GetAllActiveFilters();
+			const int32 NumFilters = FilterCollection->Num();
+			if (NumFilters == 0)
+			{
+				return true;
+			}
+		
+			for (int32 Index = 0; Index < NumFilters; ++Index)
+			{
+				if (FilterCollection->GetFilterAtIndex(Index)->PassesFilter(InItem))
+				{
+					return true;
+				}
+			}
+
+			return false;
+		};
+		
+		return TreeTextFilter->PassesFilter(InItem) && PassesAnyFilterBarFilter(InItem);
+	}
+
+	void ClearFilters()
+	{
+		// Prevent filter changed callback from firing while changing filter status, will invoke once all filters are reset
+		TGuardValue<bool> FilterChangedGuard(bFilterChangedGuard, true);
+		for (const TSharedPtr<FCompositePassTreeFilter>& Filter : PassTypeFilters)
+		{
+			Filter->DeactivateFilter();
+		}
+
+		AllPassesFilter->ActivateFilter();
+		TextFilterSearchBox->SetText(FText::GetEmpty());
+
+		OnFilterChanged.ExecuteIfBound();
+	}
+	
+private:
+	/** Creates all the filter bar filters and initializes the text filter */
+	void InitializeFilters(TArray<TSharedRef<FFilterBase<FFilterType>>>& OutAllFilters)
+	{
+		TSharedPtr<FFilterCategory> BasicFiltersCategory = MakeShared<FFilterCategory>(LOCTEXT("BasicFiltersCategory", "Basic"), FText::GetEmpty());
+
+		static TArray<TTuple<FString, FText, FText>> PassTypeFilterConfigs =
+		{
+			{ TEXT("MediaPassesFilter"), LOCTEXT("MediaPassesFilterName", "Media Passes"),  LOCTEXT("MediaPassesFilterToolTip", "Only show media passes") },
+			//{ TEXT("ScenePassesFilter"), LOCTEXT("ScenePassesFilterName", "Scene Passes"),  LOCTEXT("ScenePassesFilterToolTip", "Only show scene passes")  },
+			{ TEXT("LayerPassesFilter"), LOCTEXT("LayerPassesFilterName", "Layer Passes"),  LOCTEXT("LayerPassesFilterToolTip", "Only show layer passes")  }
+		};
+
+		for (int32 Index = 0; Index < PassTypeFilterConfigs.Num(); ++Index)
+		{
+			TSharedPtr<FCompositePassTreeFilter> PassTypeFilter = MakeShared<FCompositePassTreeFilter>(
+				BasicFiltersCategory,
+				PassTypeFilterConfigs[Index].Get<0>(),
+				PassTypeFilterConfigs[Index].Get<1>(),
+				FGenericFilter<FFilterType>::FOnItemFiltered::CreateLambda([Index](FFilterType InItem)
+				{
+					if (!InItem.IsValid())
+					{
+						return false;
+					}
+
+					return (int32)InItem->PassType == Index;
+				}));
+	
+			PassTypeFilter->SetToolTipText(PassTypeFilterConfigs[Index].Get<2>());
+			PassTypeFilter->OnActiveStateChanged.BindSP(this, &SCompositePassTreeToolbar::OnPassTypeFilterStateChanged, Index);
+
+			PassTypeFilters.Add(PassTypeFilter);
+			OutAllFilters.Add(PassTypeFilter.ToSharedRef());
+		}
+		
+		AllPassesFilter = MakeShared<FCompositePassTreeFilter>(
+				BasicFiltersCategory,
+				TEXT("AllPassesFilter"),
+				LOCTEXT("AllPassesFilterName", "All Passes"),
+				FGenericFilter<FFilterType>::FOnItemFiltered::CreateLambda([](FFilterType InItem)
+				{
+					return true;
+				}));
+	
+		AllPassesFilter->SetToolTipText(LOCTEXT("AllPassesFilterToolTip", "Show all passes"));
+		AllPassesFilter->OnActiveStateChanged.BindSP(this, &SCompositePassTreeToolbar::OnAllPassesFilterStateChanged);
+		
+		OutAllFilters.Add(AllPassesFilter.ToSharedRef());
+		
+		TreeTextFilter = MakeShared<FTreeFilter>(FTreeFilter::FItemToStringArray::CreateLambda([](const FFilterType& InItem, TArray<FString>& OutStrings)
+		{
+			if (InItem.IsValid())
+			{
+				if (InItem->HasValidPassIndex())
+				{
+					if (UCompositePassBase* Pass = InItem->GetPass())
+					{
+						OutStrings.Add(Pass->GetClass()->GetDisplayNameText().ToString());
+					}
+					else
+					{
+						OutStrings.Add(TEXT("None"));
+					}
+				}
+
+				static FString PassTypeFilterStrings[] =
+				{
+					TEXT("Media Passes"),
+					// TEXT("Scene Passes"),
+					TEXT("Layer Passes")
+				};
+				
+				OutStrings.Add(PassTypeFilterStrings[(int32)InItem->PassType]);
+			}
+		}));
+		TreeTextFilter->OnChanged().AddSP(this, &SCompositePassTreeToolbar::FilterChanged);
+	}
+	
+	TSharedRef<SWidget> GetAddPassMenuContent()
+	{
+		constexpr bool bCloseMenuAfterSelection = true;
+		FMenuBuilder MenuBuilder(bCloseMenuAfterSelection, nullptr);
+
+		TArray<UClass*> BaseLayerTypes;
+		const bool bRecursive = false;
+		GetDerivedClasses(UCompositePassBase::StaticClass(), BaseLayerTypes, bRecursive);
+
+		auto ShouldSkipClass = [](const UClass* InClass)
+			{
+				constexpr EClassFlags InvalidClassFlags = CLASS_Deprecated | CLASS_NewerVersionExists | CLASS_Abstract;
+				return InClass->HasAnyClassFlags(InvalidClassFlags) ||
+					InClass->GetName().StartsWith(TEXT("SKEL_")) ||
+					InClass->GetName().StartsWith(TEXT("REINST_"));
+			};
+		
+		// Add the native child classes of UCompositeLayerBase to the menu first
+		for (const UClass* BaseLayerType : BaseLayerTypes)
+		{
+			if (ShouldSkipClass(BaseLayerType))
+			{
+				continue;
+			}
+			
+			MenuBuilder.AddMenuEntry(
+				BaseLayerType->GetDisplayNameText(),
+				TAttribute<FText>(),
+				FSlateIconFinder::FindIconForClass(BaseLayerType),
+				FUIAction(FExecuteAction::CreateSP(this, &SCompositePassTreeToolbar::AddPass, BaseLayerType)));
+		}
+				
+		// Add any blueprint derived classes in a new section of the menu
+		bool bHasAddedSeparator = false;
+		for (const UClass* BaseLayerType : BaseLayerTypes)
+		{
+			TArray<UClass*> ChildLayerTypes;
+			GetDerivedClasses(BaseLayerType, ChildLayerTypes);
+
+			for (const UClass* ChildLayerType : ChildLayerTypes)
+			{
+				if (ShouldSkipClass(ChildLayerType))
+				{
+					continue;
+				}
+				
+				if (!bHasAddedSeparator)
+				{
+					MenuBuilder.AddSeparator();
+					bHasAddedSeparator = true;
+				}
+
+				MenuBuilder.AddMenuEntry(
+					ChildLayerType->GetDisplayNameText(),
+					TAttribute<FText>(),
+					FSlateIconFinder::FindIconForClass(ChildLayerType),
+					FUIAction(FExecuteAction::CreateSP(this, &SCompositePassTreeToolbar::AddPass, ChildLayerType)));
+			}
+		}
+		
+		return MenuBuilder.MakeWidget();
+	}
+
+	void AddPass(const UClass* InPassClass)
+	{
+		OnPassAdded.ExecuteIfBound(InPassClass);	
+	}
+
+	void OnPassTypeFilterStateChanged(bool bActive, int32 PassTypeIndex)
+	{
+		if (bActive)
+		{
+			for (int32 Index = 0; Index < PassTypeFilters.Num(); ++Index)
+			{
+				if (Index != PassTypeIndex && PassTypeFilters[Index]->IsActive())
+				{
+					PassTypeFilters[Index]->SetActive(false);
+				}
+			}
+			
+			if (AllPassesFilter->IsActive())
+			{
+				AllPassesFilter->DeactivateFilter();
+			}
+		}
+		else
+		{
+			if (!AllPassesFilter->IsActive())
+			{
+				AllPassesFilter->ActivateFilter();
+			}
+		}
+	}
+
+	void OnAllPassesFilterStateChanged(bool bActive)
+	{
+		if (bActive)
+		{
+			for (TSharedPtr<FCompositePassTreeFilter>& Filter : PassTypeFilters)
+			{
+				if (Filter->IsActive())
+				{
+					Filter->DeactivateFilter();
+				}
+			}
+		}
+		else
+		{
+			if (!PassTypeFilters[0]->IsActive())
+			{
+				PassTypeFilters[0]->ActivateFilter();
+			}
+		}
+	}
+	
+	void FilterChanged()
+	{
+		if (bFilterChangedGuard)
+		{
+			return;
+		}
+
+		OnFilterChanged.ExecuteIfBound();
+	}
+	
+private:
+	TSharedPtr<FCompositePassTreeFilter> AllPassesFilter;
+	TArray<TSharedPtr<FCompositePassTreeFilter>> PassTypeFilters;
+	TSharedPtr<FTreeFilter> TreeTextFilter;
+	
+	TSharedPtr<SFilterSearchBox> TextFilterSearchBox;
+	TSharedPtr<SCompositePassesFilterBar<FFilterType>> FilterBar;
+
+	bool bFilterChangedGuard = false;
+	
+	FOnFilterChanged OnFilterChanged;
+	FOnPassAdded OnPassAdded;
+};
+
+/** Table row that displays the pass tree item content in the tree view */
+class SCompositePassTreeItemRow : public SMultiColumnTableRow<SCompositePassTree::FPassTreeItemPtr>
+{
+	using FTreeItem = SCompositePassTree::FPassTreeItem;
+	using FTreeItemPtr = SCompositePassTree::FPassTreeItemPtr;
+	using FTreeItemWeakPtr = TWeakPtr<FTreeItem>;
+	
+public:
+	DECLARE_DELEGATE_ThreeParams(FOnPassMoved, const FTreeItemPtr& /* InItemToMove */, SCompositePassTree::EPassType /* InDestPassType */, int32 /* InDestIndex */)
+
+private:
+	/** Drag/drop operation for changing the order of passes in the tree view */
+	class FPassDragDropOp : public FDecoratedDragDropOp
+	{
+	public:
+		DRAG_DROP_OPERATOR_TYPE(FCompositeLayerDragDropOp, FDecoratedDragDropOp)
+
+		static TSharedRef<FPassDragDropOp> New(const FTreeItemPtr& InDraggedItem)
+		{
+			TSharedRef<FPassDragDropOp> Operation = MakeShared<FPassDragDropOp>();
+			Operation->DraggedItem = InDraggedItem;
+			Operation->Construct();
+			return Operation;
+		}
+
+		virtual void Construct() override
+		{
+			if (FTreeItemPtr PinnedItem = DraggedItem.Pin())
+			{
+				if (UCompositePassBase* Pass = PinnedItem->GetPass())
+				{
+					const FText PassName = Pass->GetClass()->GetDisplayNameText();
+					const FSlateBrush* PassIcon = FSlateIconFinder::FindIconForClass(Pass->GetClass()).GetIcon();
+					SetToolTip(PassName, PassIcon);
+				}
+			}
+			
+			FDecoratedDragDropOp::Construct();
+		}
+
+		FTreeItemWeakPtr DraggedItem;
+	};
+	
+public:
+	SLATE_BEGIN_ARGS(SCompositePassTreeItemRow) { }
+		SLATE_EVENT(FOnPassMoved, OnPassMoved)
+	SLATE_END_ARGS()
+
+	void Construct(const FArguments& InArgs, const TSharedRef<STableViewBase>& InOwnerTable, FTreeItemPtr InTreeItem)
+	{
+		TreeItem = InTreeItem;
+		OnPassMoved = InArgs._OnPassMoved;
+
+		STableRow<FTreeItemPtr>::FArguments Args = FSuperRowType::FArguments()
+           .Style(&FAppStyle::Get().GetWidgetStyle<FTableRowStyle>("SceneOutliner.TableViewRow"))
+           .OnDragDetected(this, &SCompositePassTreeItemRow::HandleDragDetected)
+           .OnCanAcceptDrop(this, &SCompositePassTreeItemRow::HandleCanAcceptDrop)
+           .OnAcceptDrop(this, &SCompositePassTreeItemRow::HandleAcceptDrop);
+		
+		SMultiColumnTableRow::Construct(Args, InOwnerTable);
+	}
+	
+private:
+	virtual TSharedRef<SWidget> GenerateWidgetForColumn(const FName& InColumnName) override
+	{
+		if (InColumnName == CompositePassTree::PassTreeColumn_Enabled)
+		{
+			return SNew(SCheckBox)
+				.IsChecked(this, &SCompositePassTreeItemRow::IsItemEnabled)
+				.OnCheckStateChanged(this, &SCompositePassTreeItemRow::OnIsItemEnabledChanged);
+		}
+		if (InColumnName == CompositePassTree::PassTreeColumn_Name)
+		{
+			return SNew(SHorizontalBox)
+
+				+SHorizontalBox::Slot()
+				.AutoWidth()
+				.VAlign(VAlign_Center)
+				.Padding(6.f, 0.f, 0.f, 0.f)
+				[
+					SNew(SExpanderArrow, SharedThis(this)).IndentAmount(24)
+				]
+				
+				+SHorizontalBox::Slot()
+				.AutoWidth()
+				.VAlign(VAlign_Center)
+				.Padding(0.0f, 1.0f, 6.0f, 1.0f)
+				[
+					SNew(SBox)
+					.WidthOverride(16.0f)
+					.HeightOverride(16.0f)
+					[
+						SNew(SImage).Image(this, &SCompositePassTreeItemRow::GetItemIcon)
+					]
+				]
+
+				+SHorizontalBox::Slot()
+				.VAlign(VAlign_Center)
+				[
+					SNew(STextBlock).Text(this, &SCompositePassTreeItemRow::GetItemName)
+				];
+		}
+		
+		return SNullWidget::NullWidget;
+	}
+
+	/** Gets the checkbox state for the row item's Enabled flag */
+	ECheckBoxState IsItemEnabled() const
+	{
+		if (!TreeItem->Plate.IsValid())
+		{
+			return ECheckBoxState::Unchecked;
+		}
+		
+		if (TreeItem->HasValidPassIndex())
+		{
+			if (UCompositePassBase* Pass = TreeItem->GetPass())
+			{
+				return Pass->bIsEnabled ? ECheckBoxState::Checked : ECheckBoxState::Unchecked;
+			}
+			
+			// In the case where there is a valid pass index, but the pass itself is null
+			return ECheckBoxState::Unchecked;
+		}
+		
+		// Tree item is a pass type group item, so set its checkbox state as the aggregate of all passes of that type
+		TArray<TObjectPtr<UCompositePassBase>>& AllPassesOfType = SCompositePassTree::GetPassList(TreeItem->Plate, TreeItem->PassType);
+
+		bool bAnyEnabled = false;
+		bool bAnyDisabled = false;
+		for (const TObjectPtr<UCompositePassBase>& Pass : AllPassesOfType)
+		{
+			if (Pass->bIsEnabled)
+			{
+				bAnyEnabled = true;
+			}
+			else
+			{
+				bAnyDisabled = true;
+			}
+		}
+
+		if (bAnyEnabled)
+		{
+			return bAnyDisabled ? ECheckBoxState::Undetermined : ECheckBoxState::Checked;
+		}
+		else
+		{
+			return bAnyDisabled ? ECheckBoxState::Unchecked : ECheckBoxState::Checked;
+		}
+	}
+
+	/** Raised when the row item's Enabled checkbox has changed state */
+	void OnIsItemEnabledChanged(ECheckBoxState InCheckBoxState)
+	{
+		if (!TreeItem->Plate.IsValid())
+		{
+			return;
+		}
+		
+		if (TreeItem->HasValidPassIndex())
+		{
+			if (UCompositePassBase* Pass = TreeItem->GetPass())
+			{
+				FScopedTransaction SetEnabledTransaction(LOCTEXT("SetEnabledTransaction", "Set Enabled"));
+				Pass->Modify();
+
+				Pass->bIsEnabled = InCheckBoxState == ECheckBoxState::Checked;
+				return;
+			}
+
+			// In the case where there is a valid pass index, but the pass itself is null
+			return;
+		}
+		
+		// Tree item is a pass type group item, so set enable state of all passes of type
+		TArray<TObjectPtr<UCompositePassBase>>& AllPassesOfType = SCompositePassTree::GetPassList(TreeItem->Plate, TreeItem->PassType);
+		if (AllPassesOfType.IsEmpty())
+		{
+			return;
+		}
+		
+		FScopedTransaction SetAllEnabledTransaction(LOCTEXT("SetAllEnabledTransaction", "Set All Enabled"));
+		for (const TObjectPtr<UCompositePassBase>& Pass : AllPassesOfType)
+		{
+			Pass->Modify();
+			Pass->bIsEnabled = InCheckBoxState == ECheckBoxState::Checked;
+		}
+	}
+	
+	/** Gets the row item's icon to display next to its name */
+	const FSlateBrush* GetItemIcon() const
+	{
+		if (!TreeItem->Plate.IsValid())
+		{
+			return nullptr;
+		}
+		
+		if (TreeItem->HasValidPassIndex())
+		{
+			if (UCompositePassBase* Pass = TreeItem->GetPass())
+			{
+				return FSlateIconFinder::FindIconForClass(Pass->GetClass()).GetIcon();
+			}
+			
+			// In the case where there is a valid pass index, but the pass itself is null
+			return nullptr;
+		}
+
+		// Tree item is a pass type group item, so return icon for pass type
+		static const FSlateBrush* PassTypeIcons[] =
+		{
+			FCompositeEditorStyle::Get().GetBrush("CompositeEditor.Passes.Media"),
+			//FCompositeEditorStyle::Get().GetBrush("CompositeEditor.Passes.Scene"),
+			FCompositeEditorStyle::Get().GetBrush("CompositeEditor.Passes.Layer")
+		};
+		
+		return PassTypeIcons[(int32)TreeItem->PassType];
+	}
+
+	/** Gets the row item's display name */
+	FText GetItemName() const
+	{
+		if (!TreeItem->Plate.IsValid())
+		{
+			return FText::GetEmpty();
+		}
+		
+		if (TreeItem->HasValidPassIndex())
+		{
+			if (UCompositePassBase* Pass = TreeItem->GetPass())
+			{
+				return Pass->GetClass()->GetDisplayNameText();
+			}
+			
+			// In the case where there is a valid pass index, but the pass itself is null
+			return LOCTEXT("NoneTypeLabel", "None");
+		}
+
+		// Tree item is a pass type group item, so return label for pass type
+		static FText PassTypeLabels[] =
+		{
+			LOCTEXT("MediaPassTypeLabel", "Media Passes"),
+			//LOCTEXT("ScenePassTypeLabel", "Scene Passes"),
+			LOCTEXT("LayerPassTypeLabel", "Layer Passes"),
+		};
+		
+		return PassTypeLabels[(int32)TreeItem->PassType];
+	}
+
+	/** Raised when a drag start has been detected on this row item */
+	FReply HandleDragDetected(const FGeometry& Geometry, const FPointerEvent& MouseEvent)
+	{
+		if (!TreeItem->Plate.IsValid())
+		{
+			return FReply::Unhandled();
+		}
+
+		if (!TreeItem->HasValidPassIndex())
+		{
+			return FReply::Unhandled();
+		}
+		
+		if (MouseEvent.IsMouseButtonDown(EKeys::LeftMouseButton))
+		{
+			TSharedPtr<FDragDropOperation> DragDropOp = FPassDragDropOp::New(TreeItem);
+			return FReply::Handled().BeginDragDrop(DragDropOp.ToSharedRef());
+		}
+
+		return FReply::Unhandled();
+	}
+
+	/** Raised when the user is attempting to drop something onto this item */
+	TOptional<EItemDropZone> HandleCanAcceptDrop(const FDragDropEvent& InDragDropEvent, EItemDropZone InItemDropZone, FTreeItemPtr InTreeItem)
+	{
+		TSharedPtr<FPassDragDropOp> DragDropOp = InDragDropEvent.GetOperationAs<FPassDragDropOp>();
+
+		if (!DragDropOp.IsValid() || !DragDropOp->DraggedItem.IsValid())
+		{
+			return TOptional<EItemDropZone>();
+		}
+
+		FTreeItemPtr DraggedItem = DragDropOp->DraggedItem.Pin();
+		if (DraggedItem->Plate.IsValid())
+		{
+			if (!InTreeItem->HasValidPassIndex())
+			{
+				return EItemDropZone::BelowItem;
+			}
+
+			return InItemDropZone == EItemDropZone::OntoItem ? EItemDropZone::AboveItem : InItemDropZone;
+		}
+		
+		return TOptional<EItemDropZone>();
+	}
+
+	/** Raised when the user has dropped something onto this item */
+	FReply HandleAcceptDrop(const FDragDropEvent& InDragDropEvent, EItemDropZone InItemDropZone, FTreeItemPtr InTreeItem)
+	{
+		TSharedPtr<FPassDragDropOp> DragDropOp = InDragDropEvent.GetOperationAs<FPassDragDropOp>();
+
+		if (!DragDropOp.IsValid() || !DragDropOp->DraggedItem.IsValid())
+		{
+			return FReply::Unhandled();
+		}
+
+		if (OnPassMoved.IsBound())
+		{
+			int32 DestIndex = InItemDropZone == EItemDropZone::AboveItem ? InTreeItem->PassIndex: InTreeItem->PassIndex + 1;
+
+			// If the destination is further down than the item being moved, we must subtract one from the destination index
+			// to account for the removal of the original item
+			if (InTreeItem->PassIndex > DragDropOp->DraggedItem.Pin()->PassIndex && InTreeItem->PassType == DragDropOp->DraggedItem.Pin()->PassType)
+			{
+				--DestIndex;
+			}
+			
+			OnPassMoved.Execute(DragDropOp->DraggedItem.Pin(), InTreeItem->PassType, DestIndex);
+			return FReply::Handled();
+		}
+		
+		return FReply::Unhandled();
+	}
+	
+private:
+	/** The tree item this widget is outputting */
+	FTreeItemPtr TreeItem;
+
+	/** Callback that is raised when a layer is moved via a drag and drop operation */
+	FOnPassMoved OnPassMoved;
+};
+
+bool SCompositePassTree::FPassTreeItem::HasValidPassIndex() const
+{
+	if (!Plate.IsValid())
+	{
+		return false;
+	}
+
+	TArray<TObjectPtr<UCompositePassBase>>& PassList = SCompositePassTree::GetPassList(Plate, PassType);
+	return PassList.IsValidIndex(PassIndex);
+}
+
+UCompositePassBase* SCompositePassTree::FPassTreeItem::GetPass() const
+{
+	if (!HasValidPassIndex())
+	{
+		return nullptr;
+	}
+
+	TArray<TObjectPtr<UCompositePassBase>>& PassList = SCompositePassTree::GetPassList(Plate, PassType);
+	return PassList[PassIndex];
+}
+
+void SCompositePassTree::Construct(const FArguments& InArgs, const TWeakObjectPtr<UCompositeLayerPlate>& InPlate)
+{
+	Plate = InPlate;
+	OnSelectionChanged = InArgs._OnSelectionChanged;
+	OnLayoutChanged = InArgs._OnLayoutChanged;
+
+	CommandList = MakeShared<FUICommandList>();
+	BindCommands();
+	
+	ChildSlot
+	[
+		SNew(SVerticalBox)
+
+		+SVerticalBox::Slot()
+		.AutoHeight()
+		.Padding(0.0f, 4.0f)
+		[
+			SAssignNew(Toolbar, SCompositePassTreeToolbar)
+			.OnFilterChanged(this, &SCompositePassTree::OnFilterChanged)
+			.OnPassAdded(this, &SCompositePassTree::OnPassAdded)
+		]
+
+		+SVerticalBox::Slot()
+		.AutoHeight()
+		[
+			SAssignNew(TreeView, STreeView<FPassTreeItemPtr>)
+			.TreeItemsSource(&FilteredPassTreeItems)
+			.HeaderRow(
+				SNew(SHeaderRow)
+
+				+SHeaderRow::Column(CompositePassTree::PassTreeColumn_Enabled)
+				.DefaultLabel(FText::GetEmpty())
+				.FixedWidth(24.0f)
+				.HAlignHeader(HAlign_Center)
+				.VAlignHeader(VAlign_Center)
+				.HAlignCell(HAlign_Center)
+				.VAlignCell(VAlign_Center)
+				[
+					SNew(SCheckBox)
+					.IsChecked(this, &SCompositePassTree::GetGlobalEnabledState)
+					.OnCheckStateChanged(this, &SCompositePassTree::OnGlobalEnabledStateChanged)
+				]
+
+				+SHeaderRow::Column(CompositePassTree::PassTreeColumn_Name)
+				.DefaultLabel(LOCTEXT("PassNameColumnLabel", "Pass Name"))
+			)
+			.OnGenerateRow_Lambda([this](FPassTreeItemPtr InTreeItem, const TSharedRef<STableViewBase>& InOwnerTable)
+			{
+				return SNew(SCompositePassTreeItemRow, InOwnerTable, InTreeItem)
+					.OnPassMoved(this, &SCompositePassTree::OnPassMoved);
+			})
+			.OnGetChildren_Lambda([](FPassTreeItemPtr InTreeItem, TArray<FPassTreeItemPtr>& OutChildren)
+			{
+				if (InTreeItem.IsValid())
+				{
+					for (FPassTreeItemPtr& Child : InTreeItem->Children)
+					{
+						if (Child.IsValid() && !Child->bFilteredOut)
+						{
+							OutChildren.Add(Child);
+						}
+					}
+				}
+			})
+			.OnSelectionChanged(this, &SCompositePassTree::OnPassSelectionChanged)
+			.OnContextMenuOpening(this, &SCompositePassTree::CreateTreeContextMenu)
+		]
+
+		+SVerticalBox::Slot()
+		.AutoHeight()
+		[
+			SNew(STextBlock)
+			.Text(this, &SCompositePassTree::GetFilterStatusText)
+			.ColorAndOpacity(this, &SCompositePassTree::GetFilterStatusColor)
+		]
+	];
+
+	FillPassTreeItems();
+}
+
+FReply SCompositePassTree::OnKeyDown(const FGeometry& MyGeometry, const FKeyEvent& InKeyEvent)
+{
+	if (CommandList->ProcessCommandBindings(InKeyEvent))
+	{
+		return FReply::Handled();
+	}
+	
+	return FReply::Unhandled();
+}
+
+TArray<TObjectPtr<UCompositePassBase>>& SCompositePassTree::GetPassList(TWeakObjectPtr<UCompositeLayerPlate> InPlate, EPassType InPassType)
+{
+	check(InPlate.IsValid());
+	
+	switch (InPassType)
+	{
+	case EPassType::Media:
+		return InPlate->MediaPasses;
+
+	// case EPassType::Scene:
+	// return InPlate->ScenePasses;
+		
+	case EPassType::Layer:
+	default:
+		return InPlate->LayerPasses;
+	}
+}
+
+FProperty* SCompositePassTree::GetPassListProperty(EPassType InPassType)
+{
+	switch (InPassType)
+	{
+	case EPassType::Media:
+		return FindFProperty<FProperty>(UCompositeLayerPlate::StaticClass(), GET_MEMBER_NAME_CHECKED(UCompositeLayerPlate, MediaPasses));
+
+	// case EPassType::Scene:
+	// return FindFProperty<FProperty>(UCompositeLayerPlate::StaticClass(), GET_MEMBER_NAME_CHECKED(UCompositeLayerPlate, ScenePasses));
+		
+	case EPassType::Layer:
+	default:
+		return FindFProperty<FProperty>(UCompositeLayerPlate::StaticClass(), GET_MEMBER_NAME_CHECKED(UCompositeLayerPlate, LayerPasses));
+	}
+}
+
+void SCompositePassTree::BindCommands()
+{
+	CommandList->MapAction(FGenericCommands::Get().Copy,
+		FExecuteAction::CreateSP(this, &SCompositePassTree::CopySelectedItems),
+		FCanExecuteAction::CreateSP(this, &SCompositePassTree::CanCopySelectedItems));
+
+	CommandList->MapAction(FGenericCommands::Get().Cut,
+		FExecuteAction::CreateSP(this, &SCompositePassTree::CutSelectedItems),
+		FCanExecuteAction::CreateSP(this, &SCompositePassTree::CanCutSelectedItems));
+
+	CommandList->MapAction(FGenericCommands::Get().Paste,
+		FExecuteAction::CreateSP(this, &SCompositePassTree::PasteSelectedItems),
+		FCanExecuteAction::CreateSP(this, &SCompositePassTree::CanPasteSelectedItems));
+
+	CommandList->MapAction(FGenericCommands::Get().Duplicate,
+		FExecuteAction::CreateSP(this, &SCompositePassTree::DuplicateSelectedItems),
+		FCanExecuteAction::CreateSP(this, &SCompositePassTree::CanDuplicateSelectedItems));
+	
+	CommandList->MapAction(FGenericCommands::Get().Delete,
+		FExecuteAction::CreateSP(this, &SCompositePassTree::DeleteSelectedItems),
+		FCanExecuteAction::CreateSP(this, &SCompositePassTree::CanDeleteSelectedItems));
+
+	CommandList->MapAction(FCompositeEditorCommands::Get().Enable,
+		FExecuteAction::CreateSP(this, &SCompositePassTree::EnableSelectedItems),
+		FCanExecuteAction::CreateSP(this, &SCompositePassTree::CanEnableSelectedItems));
+}
+
+void SCompositePassTree::FillPassTreeItems()
+{
+	PassTreeItems.Empty();
+
+	if (Plate.IsValid())
+	{
+		for (int32 TypeIndex = 0; TypeIndex < (int32)EPassType::MAX; ++TypeIndex)
+		{
+			const EPassType PassType = (EPassType)TypeIndex;
+			FPassTreeItemPtr NewPassTypeTreeItem = MakeShared<FPassTreeItem>();
+			NewPassTypeTreeItem->Plate = Plate;
+			NewPassTypeTreeItem->PassType = PassType;
+			
+			PassTreeItems.Add(NewPassTypeTreeItem);
+			
+			TArray<TObjectPtr<UCompositePassBase>>& Passes = GetPassList(PassType);
+			for (int32 Index = 0; Index < Passes.Num(); ++Index)
+			{
+				FPassTreeItemPtr NewPassTreeItem = MakeShared<FPassTreeItem>();
+				NewPassTreeItem->Plate = Plate;
+				NewPassTreeItem->PassType = PassType;
+				NewPassTreeItem->PassIndex = Index;
+
+				NewPassTypeTreeItem->Children.Add(NewPassTreeItem);
+			}
+		}
+	}
+
+	FilterPassTreeItems();
+	RefreshAndExpandTreeView();
+}
+
+void SCompositePassTree::FilterPassTreeItems()
+{
+	FilteredPassTreeItems.Empty();
+
+	if (!Toolbar.IsValid())
+	{
+		return;
+	}
+	
+	for (FPassTreeItemPtr& TreeItem : PassTreeItems)
+	{
+		TreeItem->bFilteredOut = !Toolbar->ItemPassesFilters(TreeItem);
+
+		bool bHasUnfilteredChildren = false;
+		for (FPassTreeItemPtr& ChildTreeItem : TreeItem->Children)
+		{
+			ChildTreeItem->bFilteredOut = !Toolbar->ItemPassesFilters(ChildTreeItem);
+			if (!ChildTreeItem->bFilteredOut)
+			{
+				bHasUnfilteredChildren = true;
+			}
+		}
+
+		if (!TreeItem->bFilteredOut || bHasUnfilteredChildren)
+		{
+			FilteredPassTreeItems.Add(TreeItem);
+		}
+	}
+}
+
+void SCompositePassTree::RefreshPassType(EPassType InPassType, bool bRefreshTreeView)
+{
+	FPassTreeItemPtr PassTypeTreeItem = PassTreeItems[(int32)InPassType];
+	PassTypeTreeItem->Children.Empty();
+
+	TArray<TObjectPtr<UCompositePassBase>>& Passes = GetPassList(InPassType);
+	for (int32 Index = 0; Index < Passes.Num(); ++Index)
+	{
+		FPassTreeItemPtr NewPassTreeItem = MakeShared<FPassTreeItem>();
+		NewPassTreeItem->Plate = Plate;
+		NewPassTreeItem->PassType = InPassType;
+		NewPassTreeItem->PassIndex = Index;
+
+		PassTypeTreeItem->Children.Add(NewPassTreeItem);
+	}
+
+	FilterPassTreeItems();
+	
+	if (bRefreshTreeView)
+	{
+		RefreshAndExpandTreeView();
+	}
+}
+
+void SCompositePassTree::RefreshAndExpandTreeView()
+{
+	if (TreeView.IsValid())
+	{
+		TreeView->RebuildList();
+
+		for (int32 Index = 0; Index < FilteredPassTreeItems.Num(); ++Index)
+		{
+			TreeView->SetItemExpansion(FilteredPassTreeItems[Index], true);
+		}
+
+		OnLayoutChanged.ExecuteIfBound();
+	}
+}
+
+TSharedPtr<SWidget> SCompositePassTree::CreateTreeContextMenu()
+{
+	if (!TreeView.IsValid())
+	{
+		return SNullWidget::NullWidget;
+	}
+	
+	const int32 NumItems = TreeView->GetNumItemsSelected();
+	if (NumItems >= 1)
+	{
+		constexpr bool bShouldCloseWindowAfterMenuSelection = true;
+		FMenuBuilder MenuBuilder(bShouldCloseWindowAfterMenuSelection, CommandList);
+
+		MenuBuilder.AddMenuEntry(FGenericCommands::Get().Copy);
+		MenuBuilder.AddMenuEntry(FGenericCommands::Get().Cut);
+		MenuBuilder.AddMenuEntry(FGenericCommands::Get().Paste);
+		MenuBuilder.AddMenuEntry(FGenericCommands::Get().Duplicate);
+		MenuBuilder.AddMenuEntry(FGenericCommands::Get().Delete);
+
+		MenuBuilder.AddSeparator();
+
+		MenuBuilder.AddMenuEntry(FCompositeEditorCommands::Get().Enable);
+		
+		return MenuBuilder.MakeWidget();
+	}
+
+	return SNullWidget::NullWidget;
+}
+
+ECheckBoxState SCompositePassTree::GetGlobalEnabledState() const
+{
+	if (!Plate.IsValid())
+	{
+		return ECheckBoxState::Checked;
+	}
+
+	// Determine if all passes are disabled, enabled, or if there is a mix
+	bool bAnyPassesEnabled = false;
+	bool bAnyPassesDisabled = false;
+	
+	for (int32 Index = 0; Index < (int32)EPassType::MAX; ++Index)
+	{
+		const TArray<TObjectPtr<UCompositePassBase>>& PassList = GetPassList((EPassType)Index);
+
+		// Consider any pass list that is empty as enabled to align with how the check state of the pass type tree item works
+		if (PassList.IsEmpty())
+		{
+			bAnyPassesEnabled = true;
+		}
+		
+		for (const TObjectPtr<UCompositePassBase>& Pass : PassList)
+		{
+			if (Pass->bIsEnabled)
+			{
+				bAnyPassesEnabled = true;
+			}
+			else
+			{
+				bAnyPassesDisabled = true;
+			}
+		}
+	}
+
+	if (bAnyPassesEnabled)
+	{
+		return bAnyPassesDisabled ? ECheckBoxState::Undetermined : ECheckBoxState::Checked;
+	}
+	else
+	{
+		return bAnyPassesDisabled ? ECheckBoxState::Unchecked : ECheckBoxState::Undetermined;
+	}
+}
+
+void SCompositePassTree::OnGlobalEnabledStateChanged(ECheckBoxState CheckBoxState)
+{
+	if (!Plate.IsValid())
+	{
+		return;
+	}
+
+	TSharedPtr<FScopedTransaction> AllEnabledTransaction;
+	
+	for (int32 Index = 0; Index < (int32)EPassType::MAX; ++Index)
+	{
+		const TArray<TObjectPtr<UCompositePassBase>>& PassList = GetPassList((EPassType)Index);
+
+		for (const TObjectPtr<UCompositePassBase>& Pass : PassList)
+		{
+			if (!AllEnabledTransaction.IsValid())
+			{
+				AllEnabledTransaction = MakeShared<FScopedTransaction>(LOCTEXT("AllEnabledTransaction", "Set All Enabled"));
+				Plate->Modify();
+			}
+
+			Pass->Modify();
+			Pass->bIsEnabled = CheckBoxState == ECheckBoxState::Checked;
+		}
+	}
+}
+
+void SCompositePassTree::OnPassMoved(const TSharedPtr<FPassTreeItem>& InTreeItem, EPassType InDestPassType, int InDestIndex)
+{
+	if (!InTreeItem.IsValid() ||
+		!InTreeItem->Plate.IsValid() ||
+		!InTreeItem->HasValidPassIndex())
+	{
+		return;
+	}
+
+	FScopedTransaction MovePassTransaction(LOCTEXT("MovePassTransaction", "Move Pass"));
+	InTreeItem->Plate->Modify();
+	
+	TArray<TObjectPtr<UCompositePassBase>>& SourcePassList = GetPassList(InTreeItem->Plate, InTreeItem->PassType);
+	TArray<TObjectPtr<UCompositePassBase>>& DestPassList = GetPassList(InTreeItem->Plate, InDestPassType);
+
+	FProperty* SourcePassListProperty = GetPassListProperty(InTreeItem->PassType);
+	FProperty* DestPassListProperty = GetPassListProperty(InDestPassType);
+	
+	UCompositePassBase* Pass = SourcePassList[InTreeItem->PassIndex];
+	
+	InTreeItem->Plate->PreEditChange(SourcePassListProperty);
+	SourcePassList.RemoveAt(InTreeItem->PassIndex);
+
+	// If we aren't adding the item back to the same pass list, invoke a property changed event for its removal
+	if (InTreeItem->PassType != InDestPassType)
+	{
+		FPropertyChangedEvent SourceListChangedEvent(SourcePassListProperty, EPropertyChangeType::ArrayRemove);
+		InTreeItem->Plate->PostEditChangeProperty(SourceListChangedEvent);
+
+		InTreeItem->Plate->PreEditChange(DestPassListProperty);
+	}
+	
+	DestPassList.Insert(Pass, InDestIndex);
+
+	// If the item moved from one list to another, invoke a property changed event for the add; otherwise, invoke a move event
+	if (InTreeItem->PassType != InDestPassType)
+	{
+		FPropertyChangedEvent DestListChangedEvent(DestPassListProperty, EPropertyChangeType::ArrayAdd);
+		InTreeItem->Plate->PostEditChangeProperty(DestListChangedEvent);
+	}
+	else
+	{
+		FPropertyChangedEvent MoveChangedEvent(SourcePassListProperty, EPropertyChangeType::ArrayMove);
+		InTreeItem->Plate->PostEditChangeProperty(MoveChangedEvent);
+	}
+	
+	RefreshPassType(InTreeItem->PassType, /* bRefreshTreeView */ false);
+	RefreshPassType(InDestPassType);
+}
+
+void SCompositePassTree::OnPassSelectionChanged(TSharedPtr<FPassTreeItem> InTreeItem, ESelectInfo::Type SelectInfo)
+{
+	TArray<FPassTreeItemPtr> SelectedItems = TreeView.IsValid() ? TreeView->GetSelectedItems() : TArray<FPassTreeItemPtr>();
+	
+	TArray<UObject*> SelectedLayers;
+	SelectedLayers.Reserve(SelectedItems.Num());
+
+	for (const FPassTreeItemPtr& Item : SelectedItems)
+	{
+		if (!Item->Plate.IsValid())
+		{
+			continue;
+		}
+
+		if (Item->HasValidPassIndex())
+		{
+			if (UCompositePassBase* Pass = Item->GetPass())
+			{
+				SelectedLayers.Add(Pass);
+			}
+		}
+	}
+	
+	OnSelectionChanged.ExecuteIfBound(SelectedLayers);
+}
+
+void SCompositePassTree::OnFilterChanged()
+{
+	FilterPassTreeItems();
+	RefreshAndExpandTreeView();
+}
+
+void SCompositePassTree::OnPassAdded(const UClass* InPassClass)
+{
+	TArray<FPassTreeItemPtr> SelectedItems = TreeView.IsValid() ? TreeView->GetSelectedItems() : TArray<FPassTreeItemPtr>();
+	if (SelectedItems.IsEmpty())
+	{
+		// If there are no selected items, add the pass to the media passes list
+		SelectedItems.Add(PassTreeItems[0]);
+	}
+
+	// Item after which the new pass will get added
+	FPassTreeItemPtr AnchorItem = SelectedItems.Last();
+	if (!AnchorItem.IsValid() || !AnchorItem->Plate.IsValid())
+	{
+		return;
+	}
+
+	FScopedTransaction AddPassTransaction(LOCTEXT("AddPassTransaction", "Add Pass"));
+	AnchorItem->Plate->Modify();
+	
+	UCompositePassBase* NewPass = NewObject<UCompositePassBase>(AnchorItem->Plate.Get(), InPassClass);
+	TArray<TObjectPtr<UCompositePassBase>>& PassList = GetPassList(AnchorItem->Plate, AnchorItem->PassType);
+
+	FProperty* PassListProperty = GetPassListProperty(AnchorItem->PassType);
+	AnchorItem->Plate->PreEditChange(PassListProperty);
+	
+	int32 NewPassIndex = PassList.Insert(NewPass, AnchorItem->PassIndex + 1);
+
+	FPropertyChangedEvent PropertyChangedEvent(PassListProperty, EPropertyChangeType::ArrayAdd);
+	AnchorItem->Plate->PostEditChangeProperty(PropertyChangedEvent);
+
+	if (Toolbar.IsValid())
+	{
+		Toolbar->ClearFilters();
+	}
+	
+	RefreshPassType(AnchorItem->PassType);
+
+	FPassTreeItemPtr PassToSelect = PassTreeItems[(int32)AnchorItem->PassType]->Children[NewPassIndex];
+	if (TreeView.IsValid())
+	{
+		TreeView->ClearSelection();
+		TreeView->SetItemSelection(PassToSelect, true);
+	}
+}
+
+FText SCompositePassTree::GetFilterStatusText() const
+{
+	int32 NumPasses = 0;
+	for (const FPassTreeItemPtr& PassTypeItem : PassTreeItems)
+	{
+		NumPasses += PassTypeItem->Children.Num();
+	}
+
+	int32 NumFilteredPasses = 0;
+	for (const FPassTreeItemPtr& PassTypeItem : FilteredPassTreeItems)
+	{
+		for (const FPassTreeItemPtr& PassItem : PassTypeItem->Children)
+		{
+			if (!PassItem->bFilteredOut)
+			{
+				++NumFilteredPasses;
+			}
+		}
+	}
+	
+	TArray<FPassTreeItemPtr> SelectedItems = TreeView.IsValid() ? TreeView->GetSelectedItems() : TArray<FPassTreeItemPtr>();
+	int32 NumSelectedPasses = 0;
+	for (const FPassTreeItemPtr& PassItem : SelectedItems)
+	{
+		if (PassItem->HasValidPassIndex())
+		{
+			++NumSelectedPasses;
+		}
+	}
+	
+	const FText PassLabel = NumPasses > 1 ? LOCTEXT("PassPlural", "Passes") : LOCTEXT("PassSingular", "Pass");
+	if (NumPasses > 0 && NumPasses == NumFilteredPasses)
+	{
+		if (NumSelectedPasses > 0)
+		{
+			return FText::Format(LOCTEXT("NumPassesAndSelectedTextFormat", "{0} {1}, {2} Selected"), FText::AsNumber(NumPasses), PassLabel, FText::AsNumber(NumSelectedPasses));
+		}
+		else
+		{
+			return FText::Format(LOCTEXT("NumPassesTextFormat", "{0} {1}"), FText::AsNumber(NumPasses), PassLabel);
+		}
+	}
+	else if (NumFilteredPasses > 0)
+	{
+		if (NumSelectedPasses > 0)
+		{
+			return FText::Format(LOCTEXT("NumPassesWithFilteredAndSelectedTextFormat", "Showing {0} of {1} {2}, {3} Selected"),
+				FText::AsNumber(NumSelectedPasses),
+				FText::AsNumber(NumPasses),
+				PassLabel,
+				FText::AsNumber(NumSelectedPasses));
+		}
+		else
+		{
+			return FText::Format(LOCTEXT("NumPassesWithFilteredTextFormat", "Showing {0} of {1} {2}"),
+				FText::AsNumber(NumFilteredPasses),
+				FText::AsNumber(NumPasses),
+				PassLabel);
+		}
+	}
+	else
+	{
+		if (NumPasses > 0)
+		{
+			return FText::Format(LOCTEXT("NoMatchingPassesTextFormat", "No matching passes ({0} {1})"), FText::AsNumber(NumPasses), PassLabel);
+		}
+		else
+		{
+			return LOCTEXT("NoPassesLabel", "No passes on plate");
+		}
+	}
+}
+
+FSlateColor SCompositePassTree::GetFilterStatusColor() const
+{
+	int32 NumPasses = 0;
+	for (const FPassTreeItemPtr& PassTypeItem : PassTreeItems)
+	{
+		NumPasses += PassTypeItem->Children.Num();
+	}
+
+	int32 NumFilteredPasses = 0;
+	for (const FPassTreeItemPtr& PassTypeItem : FilteredPassTreeItems)
+	{
+		for (const FPassTreeItemPtr& PassItem : PassTypeItem->Children)
+		{
+			if (!PassItem->bFilteredOut)
+			{
+				++NumFilteredPasses;
+			}
+		}
+	}
+
+	if (NumPasses == 0 || NumFilteredPasses == 0)
+	{
+		return FAppStyle::Get().GetSlateColor("Colors.AccentRed");
+	}
+	
+	return FSlateColor::UseForeground();
+}
+
+void SCompositePassTree::CopySelectedItems()
+{
+	TArray<UObject*> ObjectsToCopy;
+
+	TArray<FPassTreeItemPtr> SelectedItems = TreeView.IsValid() ? TreeView->GetSelectedItems() : TArray<FPassTreeItemPtr>();
+	for (const FPassTreeItemPtr& SelectedItem : SelectedItems)
+	{
+		if (SelectedItem.IsValid() && SelectedItem->HasValidPassIndex())
+		{
+			if (UCompositePassBase* Pass = SelectedItem->GetPass())
+			{
+				ObjectsToCopy.Add(Pass);
+			}
+		}
+	}
+
+	if (ObjectsToCopy.IsEmpty())
+	{
+		return;
+	}
+	
+	UnMarkAllObjects(EObjectMark(OBJECTMARK_TagExp | OBJECTMARK_TagImp));
+
+	FStringOutputDevice Archive;
+	const FExportObjectInnerContext Context;
+
+	for (UObject* Object : ObjectsToCopy)
+	{
+		UExporter::ExportToOutputDevice(&Context, Object, nullptr, Archive, TEXT("copy"), 0, PPF_ExportsNotFullyQualified | PPF_Copy | PPF_Delimited, false, nullptr);
+	}
+
+	FPlatformApplicationMisc::ClipboardCopy(*Archive);
+}
+
+bool SCompositePassTree::CanCopySelectedItems()
+{
+	TArray<FPassTreeItemPtr> SelectedItems = TreeView.IsValid() ? TreeView->GetSelectedItems() : TArray<FPassTreeItemPtr>();
+	if (!SelectedItems.Num())
+	{
+		return false;
+	}
+
+	const bool bContainsPass = SelectedItems.ContainsByPredicate([](const FPassTreeItemPtr& InTreeItem)
+	{
+		return InTreeItem.IsValid() && InTreeItem->HasValidPassIndex() && InTreeItem->GetPass();
+	});
+
+	return bContainsPass;
+}
+
+void SCompositePassTree::CutSelectedItems()
+{
+	CopySelectedItems();
+	DeleteSelectedItems();
+}
+
+bool SCompositePassTree::CanCutSelectedItems()
+{
+	return CanCopySelectedItems() && CanDeleteSelectedItems();
+}
+
+class FCompositePassObjectTextFactory : public FCustomizableTextObjectFactory
+{
+public:
+	FCompositePassObjectTextFactory() : FCustomizableTextObjectFactory(GWarn) { }
+
+	virtual bool CanCreateClass(UClass* ObjectClass, bool& bOmitSubObjs) const override
+	{
+		return ObjectClass->IsChildOf(UCompositePassBase::StaticClass());
+	}
+
+	virtual void ProcessConstructedObject(UObject* NewObject) override
+	{
+		check(NewObject);
+
+		if (UCompositePassBase* Pass = Cast<UCompositePassBase>(NewObject))
+		{
+			CompositePasses.Add(Pass);
+		}
+	}
+
+public:
+	TArray<UCompositePassBase*> CompositePasses;
+};
+
+void SCompositePassTree::PasteSelectedItems()
+{
+	if (!Plate.IsValid())
+	{
+		return;
+	}
+	
+	TArray<FPassTreeItemPtr> SelectedItems = TreeView.IsValid() ? TreeView->GetSelectedItems() : TArray<FPassTreeItemPtr>();
+	if (SelectedItems.IsEmpty())
+	{
+		return;
+	}
+
+	// A list of pass types to paste the clipboard contents in, and the index to paste at
+	TMap<EPassType, int32> PassTypesToPasteTo;
+	for (const FPassTreeItemPtr& SelectedItem : SelectedItems)
+	{
+		if (SelectedItem.IsValid())
+		{
+			if (PassTypesToPasteTo.Contains(SelectedItem->PassType))
+			{
+				PassTypesToPasteTo[SelectedItem->PassType] = FMath::Max(SelectedItem->PassIndex + 1, PassTypesToPasteTo[SelectedItem->PassType]);
+			}
+			else
+			{
+				PassTypesToPasteTo.Add(SelectedItem->PassType, SelectedItem->PassIndex + 1);
+			}
+		}
+	}
+
+	if (PassTypesToPasteTo.IsEmpty())
+	{
+		return;
+	}
+
+	FString ClipboardContent;
+	FPlatformApplicationMisc::ClipboardPaste(ClipboardContent);
+	
+	FCompositePassObjectTextFactory Factory;
+	Factory.ProcessBuffer(GetTransientPackage(), RF_Transactional | RF_Transient, ClipboardContent);
+
+	if (Factory.CompositePasses.IsEmpty())
+	{
+		return;
+	}
+
+	TMap<EPassType, TArray<int32>> PassesToSelect;
+	FScopedTransaction PastePassesTransaction(LOCTEXT("PastePassesTransaction", "Paste Passes"));
+	Plate->Modify();
+	for (const TPair<EPassType, int32>& Pair : PassTypesToPasteTo)
+	{
+		EPassType PassTypeToPasteAt = Pair.Key;
+		int32 IndexToPasteAt = Pair.Value;
+
+		TArray<TObjectPtr<UCompositePassBase>>& PassList = GetPassList(PassTypeToPasteAt);
+		
+		PassesToSelect.Add(PassTypeToPasteAt, TArray<int32>());
+		
+		for (UCompositePassBase* PassToPaste : Factory.CompositePasses)
+		{
+			UCompositePassBase* NewPass = DuplicateObject(PassToPaste, Plate.Get());
+			
+			PassList.Insert(NewPass, IndexToPasteAt);
+			PassesToSelect[PassTypeToPasteAt].Add(IndexToPasteAt);
+			++IndexToPasteAt;
+		}
+
+		constexpr bool bRefreshTree = false;
+		RefreshPassType(PassTypeToPasteAt, bRefreshTree);
+	}
+
+	FilterPassTreeItems();
+	RefreshAndExpandTreeView();
+
+	if (TreeView.IsValid())
+	{
+		for (const TPair<EPassType, TArray<int32>>& Pair : PassesToSelect)
+		{
+			EPassType PassType = Pair.Key;
+			for (int32 IndexToSelect : Pair.Value)
+			{
+				if (PassTreeItems[(int32)PassType]->Children.IsValidIndex(IndexToSelect))
+				{
+					TreeView->SetItemSelection(PassTreeItems[(int32)PassType]->Children[IndexToSelect], true);
+				}
+			}
+		}
+	}
+}
+
+bool SCompositePassTree::CanPasteSelectedItems()
+{
+	TArray<FPassTreeItemPtr> SelectedItems = TreeView.IsValid() ? TreeView->GetSelectedItems() : TArray<FPassTreeItemPtr>();
+	if (SelectedItems.IsEmpty())
+	{
+		return false;
+	}
+	
+	// Can't paste unless the clipboard has a string in it
+	FString ClipboardContent;
+	FPlatformApplicationMisc::ClipboardPaste(ClipboardContent);
+	if (!ClipboardContent.IsEmpty())
+	{
+		FCompositePassObjectTextFactory Factory;
+		Factory.ProcessBuffer(GetTransientPackage(), RF_Transactional | RF_Transient, ClipboardContent);
+		return !Factory.CompositePasses.IsEmpty();
+	}
+
+	return false;
+}
+
+void SCompositePassTree::DuplicateSelectedItems()
+{
+	CopySelectedItems();
+	PasteSelectedItems();
+}
+
+bool SCompositePassTree::CanDuplicateSelectedItems()
+{
+	return CanCopySelectedItems();
+}
+
+void SCompositePassTree::DeleteSelectedItems()
+{
+	if (!Plate.IsValid())
+	{
+		return;
+	}
+	
+	TArray<FPassTreeItemPtr> SelectedItems = TreeView.IsValid() ? TreeView->GetSelectedItems() : TArray<FPassTreeItemPtr>();
+
+	TMap<EPassType, TArray<int32>> PassesToDelete;
+	for (const FPassTreeItemPtr& Item : SelectedItems)
+	{
+		if (!Item.IsValid() || !Item->HasValidPassIndex())
+		{
+			continue;
+		}
+
+		if (PassesToDelete.Contains(Item->PassType))
+		{
+			PassesToDelete[Item->PassType].Add(Item->PassIndex);
+		}
+		else
+		{
+			PassesToDelete.Add(Item->PassType, TArray<int32> { Item->PassIndex });
+		}
+	}
+
+	if (PassesToDelete.Num())
+	{
+		FScopedTransaction DeletePassTransaction(LOCTEXT("DeletePassTransaction", "Delete Pass"));
+		Plate->Modify();
+		
+		for (TPair<EPassType, TArray<int32>>& PassesToDeleteByType : PassesToDelete)
+		{
+			const EPassType PassType = PassesToDeleteByType.Key;
+			TArray<TObjectPtr<UCompositePassBase>>& PassList = GetPassList(PassType);
+			
+			// Delete the passes in reverse order so that indices don't get messed up as we delete
+			TArray<int32>& PassIndicesToDelete = PassesToDeleteByType.Value;
+			PassIndicesToDelete.Sort();
+			
+			for (int32 Index = PassIndicesToDelete.Num() - 1; Index >= 0; --Index)
+			{
+				PassList.RemoveAt(PassIndicesToDelete[Index], EAllowShrinking::No);
+			}
+
+			PassList.Shrink();
+			
+			constexpr bool bRefreshTreeView = false;
+			RefreshPassType(PassType, bRefreshTreeView);
+		}
+
+		RefreshAndExpandTreeView();
+	}
+}
+
+bool SCompositePassTree::CanDeleteSelectedItems()
+{
+	TArray<FPassTreeItemPtr> SelectedItems = TreeView.IsValid() ? TreeView->GetSelectedItems() : TArray<FPassTreeItemPtr>();
+	const bool bContainsPassInSelection = SelectedItems.ContainsByPredicate([](const FPassTreeItemPtr& InTreeItem)
+	{
+		return InTreeItem.IsValid() && InTreeItem->HasValidPassIndex();
+	});
+	
+	return bContainsPassInSelection;
+}
+
+void SCompositePassTree::EnableSelectedItems()
+{
+	TArray<FPassTreeItemPtr> SelectedItems = TreeView.IsValid() ? TreeView->GetSelectedItems() : TArray<FPassTreeItemPtr>();
+
+	bool bAnyEnabled = false;
+	for (const FPassTreeItemPtr& Item : SelectedItems)
+	{
+		if (Item->HasValidPassIndex())
+		{
+			if (UCompositePassBase* Pass = Item->GetPass())
+			{
+				if (Pass->bIsEnabled)
+				{
+					bAnyEnabled = true;
+				}
+			}
+		}
+		else
+		{
+			for (const FPassTreeItemPtr& ChildItem : Item->Children)
+			{
+				if (ChildItem->HasValidPassIndex())
+				{
+					if (UCompositePassBase* Pass = ChildItem->GetPass())
+					{
+						if (Pass->bIsEnabled)
+						{
+							bAnyEnabled = true;
+						}
+					}
+				}
+			}
+		}
+	}
+
+	const bool bSetEnabled = !bAnyEnabled;
+
+	FScopedTransaction SetEnabledTransaction(LOCTEXT("SetEnabledTransaction", "Set Enabled"));
+	Plate->Modify();
+
+	for (const FPassTreeItemPtr& Item : SelectedItems)
+	{
+		if (Item->HasValidPassIndex())
+		{
+			if (UCompositePassBase* Pass = Item->GetPass())
+			{
+				Pass->Modify();
+				Pass->bIsEnabled = bSetEnabled;
+			}
+		}
+		else
+		{
+			for (const FPassTreeItemPtr& ChildItem : Item->Children)
+			{
+				if (ChildItem->HasValidPassIndex())
+				{
+					if (UCompositePassBase* Pass = ChildItem->GetPass())
+					{
+						Pass->Modify();
+						Pass->bIsEnabled = bSetEnabled;
+					}
+				}
+			}
+		}
+	}
+}
+
+bool SCompositePassTree::CanEnableSelectedItems()
+{
+	TArray<FPassTreeItemPtr> SelectedItems = TreeView.IsValid() ? TreeView->GetSelectedItems() : TArray<FPassTreeItemPtr>();
+	return !SelectedItems.IsEmpty();
+}
+
+#undef LOCTEXT_NAMESPACE
diff --git a/Engine/Plugins/Experimental/Compositing/Composite/Source/CompositeEditor/Private/UI/SCompositePassTree.h b/Engine/Plugins/Experimental/Compositing/Composite/Source/CompositeEditor/Private/UI/SCompositePassTree.h
new file mode 100644
--- /dev/null
+++ b/Engine/Plugins/Experimental/Compositing/Composite/Source/CompositeEditor/Private/UI/SCompositePassTree.h
@@ -0,0 +1,162 @@
+﻿// Copyright Epic Games, Inc. All Rights Reserved.
+
+#pragma once
+
+#include "Widgets/SCompoundWidget.h"
+
+class FUICommandList;
+class SCompositePassTreeToolbar;
+class UCompositePassBase;
+class UCompositeLayerPlate;
+template<typename T> class STreeView;
+
+/**
+ * Widget that displays a filterable tree view for the passes within a composite plate layer
+ */
+class SCompositePassTree : public SCompoundWidget
+{
+public:
+	DECLARE_DELEGATE_OneParam(FOnSelectionChanged, const TArray<UObject*>&);
+	
+private:
+	enum class EPassType : uint8
+	{
+		Media,
+		//Scene, // TODO: Add Scene when it is supported
+		Layer,
+		MAX
+	};
+
+	/** Stores data needed for a single row item in the tree view */
+	struct FPassTreeItem
+	{
+		/** Plate whose passes are being displayed */
+		TWeakObjectPtr<UCompositeLayerPlate> Plate = nullptr;
+
+		/** The pass type this tree item belongs to */
+		EPassType PassType = EPassType::Media;
+
+		/** If this tree item is for a specific pass, this is the pass's index in the plate's corresponding list of passes */
+		int32 PassIndex = INDEX_NONE;
+
+		/** Indicates if this tree item is being filtered out by the current filter */
+		bool bFilteredOut = false;
+
+		/** Child tree items of this tree item */
+		TArray<TSharedPtr<FPassTreeItem>> Children;
+		
+		
+		/** Gets whether the pass index is valid within the plate's list of passes */
+		bool HasValidPassIndex() const;
+
+		/** Gets the pass object referenced by this tree item */
+		UCompositePassBase* GetPass() const;
+	};
+	using FPassTreeItemPtr = TSharedPtr<FPassTreeItem>;
+	
+public:
+	SLATE_BEGIN_ARGS(SCompositePassTree) { }
+		SLATE_EVENT(FOnSelectionChanged, OnSelectionChanged)
+		SLATE_EVENT(FSimpleDelegate, OnLayoutChanged)
+	SLATE_END_ARGS()
+
+	void Construct(const FArguments& InArgs, const TWeakObjectPtr<UCompositeLayerPlate>& InPlate);
+
+	// SWidget interface
+	virtual FReply OnKeyDown(const FGeometry& MyGeometry, const FKeyEvent& InKeyEvent) override;
+	// End of SWidget interface
+	
+private:
+	/** Gets a reference to the pass list that corresponds to the specified pass type in the plate. The plate pointer must be valid to call properly */
+	static TArray<TObjectPtr<UCompositePassBase>>& GetPassList(TWeakObjectPtr<UCompositeLayerPlate> InPlate, EPassType InPassType);
+	TArray<TObjectPtr<UCompositePassBase>>& GetPassList(EPassType InPassType) { return GetPassList(Plate, InPassType); }
+	const TArray<TObjectPtr<UCompositePassBase>>& GetPassList(EPassType InPassType) const { return GetPassList(Plate, InPassType); }
+	
+	/** Gets the FProperty for the specified pass type */
+	static FProperty* GetPassListProperty(EPassType InPassType);
+
+	/** Binds commands in the command list to callbacks */
+	void BindCommands();
+	
+	/** Creates any tree items needed for all passes in the plate */
+	void FillPassTreeItems();
+
+	/** Applies the current filter to the pass tree items to build a filtered list of tree items */
+	void FilterPassTreeItems();
+
+	/** Refreshes the tree items for the passes in the specified pass type list */
+	void RefreshPassType(EPassType InPassType, bool bRefreshTreeView = true);
+	
+	/** Refreshes the tree view's displayed list and automatically expands all tree items */
+	void RefreshAndExpandTreeView();
+
+	/** Creates the tree view right click context menu */
+	TSharedPtr<SWidget> CreateTreeContextMenu();
+
+	/** Gets the global enabled state for all passes on the plate */
+	ECheckBoxState GetGlobalEnabledState() const;
+
+	/** Sets the enabled state for all passes on the plate */
+	void OnGlobalEnabledStateChanged(ECheckBoxState CheckBoxState);
+	
+	/** Raised when a pass has been moved via a drag drop operation */
+	void OnPassMoved(const TSharedPtr<FPassTreeItem>& InTreeItem, EPassType InDestPassType, int InDestIndex);
+
+	/** Raised when the selected items in the tree view have bene changed */
+	void OnPassSelectionChanged(TSharedPtr<FPassTreeItem> InTreeItem, ESelectInfo::Type SelectInfo);
+	
+	/** Raised when the active filter has been changed */
+	void OnFilterChanged();
+
+	/** Raised when a new pass is being added */
+	void OnPassAdded(const UClass* InPassClass);
+
+	/** Gets status text based on selection and active filter */
+	FText GetFilterStatusText() const;
+
+	/** Gets status text color based on selection and active filter */
+	FSlateColor GetFilterStatusColor() const;
+	
+	/** Generic command callbacks to handle copy, cut, paste, delete, and rename */
+	void CopySelectedItems();
+	bool CanCopySelectedItems();
+	void CutSelectedItems();
+	bool CanCutSelectedItems();
+	void PasteSelectedItems();
+	bool CanPasteSelectedItems();
+	void DuplicateSelectedItems();
+	bool CanDuplicateSelectedItems();
+	void DeleteSelectedItems();
+	bool CanDeleteSelectedItems();
+	void EnableSelectedItems();
+	bool CanEnableSelectedItems();
+	
+private:
+	/** The plate whose passes are being displayed */
+	TWeakObjectPtr<UCompositeLayerPlate> Plate;
+
+	/** List of all pass tree items for the plate */
+	TArray<FPassTreeItemPtr> PassTreeItems;
+
+	/** List of all passes that match the current filter */
+	TArray<FPassTreeItemPtr> FilteredPassTreeItems;
+	
+	/** Tree view that displays a hierarchical list of passes on the plate */
+	TSharedPtr<STreeView<FPassTreeItemPtr>> TreeView;
+
+	/** Toolbar widget displayed above the tree view */
+	TSharedPtr<SCompositePassTreeToolbar> Toolbar;
+
+	/** The command list for commands related to the tree view */
+	TSharedPtr<FUICommandList> CommandList;
+	
+	/** Delegate that gets raised when the tree view selection changes */
+	FOnSelectionChanged OnSelectionChanged;
+
+	/** Delegate raised when there is a potential layout change to the tree view, to give containers the opportunity to correctly resize their layouts */
+	FSimpleDelegate OnLayoutChanged;
+	
+	friend class SCompositePassTreeItemRow;
+	friend class SCompositePassTreeToolbar;
+	friend class FCompositePassTreeFilter;
+};
diff --git a/Engine/Plugins/Experimental/Compositing/Composite/Source/CompositeEditor/Private/UI/SCompositePlatePassPanel.cpp b/Engine/Plugins/Experimental/Compositing/Composite/Source/CompositeEditor/Private/UI/SCompositePlatePassPanel.cpp
new file mode 100644
--- /dev/null
+++ b/Engine/Plugins/Experimental/Compositing/Composite/Source/CompositeEditor/Private/UI/SCompositePlatePassPanel.cpp
@@ -0,0 +1,88 @@
+﻿// Copyright Epic Games, Inc. All Rights Reserved.
+
+
+#include "SCompositePlatePassPanel.h"
+
+#include "DetailsViewArgs.h"
+#include "PropertyEditorModule.h"
+#include "PropertyPath.h"
+#include "SCompositePassTree.h"
+#include "Layers/CompositeLayerPlate.h"
+#include "Modules/ModuleManager.h"
+
+void SCompositePlatePassPanel::Construct(const FArguments& InArgs, UCompositeLayerPlate* InPlate)
+{
+	Plate = InPlate;
+	OnLayoutSizeChanged = InArgs._OnLayoutSizeChanged;
+	
+	FDetailsViewArgs DetailsViewArgs;
+	DetailsViewArgs.NameAreaSettings = FDetailsViewArgs::HideNameArea;
+	DetailsViewArgs.bAllowMultipleTopLevelObjects = true;
+	DetailsViewArgs.bShowPropertyMatrixButton = false;
+	DetailsViewArgs.bShowScrollBar = false;
+	DetailsViewArgs.bAllowSearch = false;
+	
+	DetailsView = FModuleManager::GetModuleChecked<FPropertyEditorModule>("PropertyEditor").CreateDetailView(DetailsViewArgs);
+
+	DetailsView->CountRows();
+	
+	ChildSlot
+	[
+		SNew(SVerticalBox)
+
+		+SVerticalBox::Slot()
+		.AutoHeight()
+		[
+			SAssignNew(TreeView, SCompositePassTree, Plate)
+			.OnSelectionChanged(this, &SCompositePlatePassPanel::OnTreeViewSelectionChanged)
+			.OnLayoutChanged_Lambda([this]
+			{
+				OnLayoutSizeChanged.ExecuteIfBound();
+			})
+		]
+
+		+SVerticalBox::Slot()
+		.AutoHeight()
+		[
+			DetailsView.ToSharedRef()
+		]
+	];
+}
+
+void SCompositePlatePassPanel::Tick(const FGeometry& AllottedGeometry, const double InCurrentTime, const float InDeltaTime)
+{
+	SCompoundWidget::Tick(AllottedGeometry, InCurrentTime, InDeltaTime);
+
+	// This panel often changes its size dynamically due to the details panel (showing new objects, expanding properties, etc.)
+	// When this panel is embedded into a parent details panel (or any list/tree view), the dynamic size is an issue when it comes to
+	// determining scroll bar visibility and size. As such, this panel invokes a OnLayoutSizeChange event that allows parent containers
+	// to request a layout refresh, where the correct scroll bar size can be computed. To detect if the details view may have changed size,
+	// we can use GetPropertyRowNumbers, which returns a list of properties that are actively being displayed as widgets (thus will exclude collapsed
+	// properties). Caching this count on tick allows this panel to detect if the object properties being displayed have changed or been expanded.
+	// The actual layout size changed event must be invoked one tick after a row count difference has been detected to give the details view
+	// time to actually change its rendered geometry
+	
+	if (bLayoutSizeChanged)
+	{
+		OnLayoutSizeChanged.ExecuteIfBound();
+		bLayoutSizeChanged = false;
+	}
+	
+	const int32 RowCount = DetailsView->GetPropertyRowNumbers().Num();	
+ 	if (CachedDetailsViewRowCount != RowCount)
+	{
+		CachedDetailsViewRowCount = RowCount;
+ 		bLayoutSizeChanged = true;
+	}
+}
+
+void SCompositePlatePassPanel::OnTreeViewSelectionChanged(const TArray<UObject*>& Objects)
+{
+	if (!DetailsView.IsValid())
+	{
+		return;
+	}
+
+	DetailsView->SetObjects(Objects);
+}
+
diff --git a/Engine/Plugins/Experimental/Compositing/Composite/Source/CompositeEditor/Private/UI/SCompositePlatePassPanel.h b/Engine/Plugins/Experimental/Compositing/Composite/Source/CompositeEditor/Private/UI/SCompositePlatePassPanel.h
new file mode 100644
--- /dev/null
+++ b/Engine/Plugins/Experimental/Compositing/Composite/Source/CompositeEditor/Private/UI/SCompositePlatePassPanel.h
@@ -0,0 +1,50 @@
+﻿// Copyright Epic Games, Inc. All Rights Reserved.
+
+#pragma once
+
+#include "Widgets/SCompoundWidget.h"
+
+class IDetailsView;
+class SCompositePassTree;
+class UCompositePassBase;
+class UCompositeLayerPlate;
+
+/**
+ * A panel that displays a tree view and a details panel for passes within the plate composite layer
+ */
+class SCompositePlatePassPanel : public SCompoundWidget
+{
+public:
+	SLATE_BEGIN_ARGS(SCompositePlatePassPanel) { }
+		SLATE_EVENT(FSimpleDelegate, OnLayoutSizeChanged)
+	SLATE_END_ARGS()
+
+	void Construct(const FArguments& InArgs, UCompositeLayerPlate* InPlate);
+
+	// SWidget interface
+	virtual void Tick(const FGeometry& AllottedGeometry, const double InCurrentTime, const float InDeltaTime) override;
+	// End of SWidget interface
+	
+private:
+	/** Raised when the selected items in the tree view have changed */
+	void OnTreeViewSelectionChanged(const TArray<UObject*>& Objects);
+	
+private:
+	/** The plate whose passes are being displayed in the panel */
+	TWeakObjectPtr<UCompositeLayerPlate> Plate;
+	
+	/** Tree view that displays a hierarchical list of passes in the plate */
+	TSharedPtr<SCompositePassTree> TreeView;
+
+	/** The details view to display selected pass properties */
+	TSharedPtr<IDetailsView> DetailsView;
+
+	/** To identify when the details view has changed size by showing a new object or expanding a property group, keep track of its reported row count every tick */
+	int32 CachedDetailsViewRowCount = 0;
+
+	/** Flag used to indicate if OnLayoutSizeChanged needs to be called on the next tick */
+	bool bLayoutSizeChanged = false;
+	
+	/** Delegate that is raised when a potential size change has been detected in the panel (e.g. the details panel is displaying a new object) */
+	FSimpleDelegate OnLayoutSizeChanged;
+};
diff --git a/Engine/Plugins/Experimental/Compositing/Composite/Source/CompositeEditor/Public/CompositeEditorModule.h b/Engine/Plugins/Experimental/Compositing/Composite/Source/CompositeEditor/Public/CompositeEditorModule.h
--- a/Engine/Plugins/Experimental/Compositing/Composite/Source/CompositeEditor/Public/CompositeEditorModule.h
+++ b/Engine/Plugins/Experimental/Compositing/Composite/Source/CompositeEditor/Public/CompositeEditorModule.h
@@ -28,6 +28,9 @@ private:
 	void RegisterTabSpawners();
 	void UnregisterTabSpawners();
 
+	void RegisterCustomizations();
+	void UnregisterCustomizations();
+
 private:
 	FDelegateHandle OnTabManagerChangedDelegateHandle;
 };
