Composite: Hidding pass & layer enabled boolean behind getter & setter.

* Fixes missing state changes from direct editor code state change.

diff --git a/Engine/Plugins/Experimental/Compositing/Composite/Source/Composite/Private/CompositeActor.cpp b/Engine/Plugins/Experimental/Compositing/Composite/Source/Composite/Private/CompositeActor.cpp
--- a/Engine/Plugins/Experimental/Compositing/Composite/Source/Composite/Private/CompositeActor.cpp
+++ b/Engine/Plugins/Experimental/Compositing/Composite/Source/Composite/Private/CompositeActor.cpp
@@ -59,7 +59,7 @@ namespace UE
 				// Iterate backwards so that lower passes in the UI are executed first
 				for (const TObjectPtr<UCompositeLayerBase>& Layer : InLayers)
 				{
-					if (!IsValid(Layer) || !Layer->bIsEnabled)
+					if (!IsValid(Layer) || !Layer->IsEnabled())
 					{
 						continue;
 					}
@@ -86,7 +86,7 @@ namespace UE
 
 				for (const TObjectPtr<UCompositeLayerBase>& Layer : InLayers)
 				{
-					if (!IsValid(Layer) || !Layer->bIsEnabled)
+					if (!IsValid(Layer) || !Layer->IsEnabled())
 					{
 						continue;
 					}
@@ -230,7 +230,7 @@ ACompositeActor::ACompositeActor(const FObjectInitializer& ObjectInitializer)
 	CompositeLayers.Add(CreateDefaultSubobject<UCompositeLayerPlate>(TEXT("PlateLayer")));
 
 	// By default, we disable the shadow reflection catcher due to its significant added cost.
-	CompositeLayers[1]->bIsEnabled = false;
+	CompositeLayers[1]->SetEnabled(false);
 }
 
 ACompositeActor::~ACompositeActor() = default;
@@ -406,7 +406,7 @@ void ACompositeActor::PostEditChangeChainProperty(struct FPropertyChangedChainEv
 			{
 				const TObjectPtr<UCompositeLayerBase>& Layer = CompositeLayers[Index];
 
-				if (!IsValid(Layer) || !Layer->bIsEnabled)
+				if (!IsValid(Layer) || !Layer->IsEnabled())
 				{
 					continue;
 				}
diff --git a/Engine/Plugins/Experimental/Compositing/Composite/Source/Composite/Private/Layers/CompositeLayerBase.cpp b/Engine/Plugins/Experimental/Compositing/Composite/Source/Composite/Private/Layers/CompositeLayerBase.cpp
--- a/Engine/Plugins/Experimental/Compositing/Composite/Source/Composite/Private/Layers/CompositeLayerBase.cpp
+++ b/Engine/Plugins/Experimental/Compositing/Composite/Source/Composite/Private/Layers/CompositeLayerBase.cpp
@@ -57,6 +57,16 @@ void UCompositeLayerBase::PostEditChangeProperty(FPropertyChangedEvent& Property
 }
 #endif //WITH_EDITOR
 
+bool UCompositeLayerBase::IsEnabled() const
+{
+	return bIsEnabled;
+}
+
+void UCompositeLayerBase::SetEnabled(bool bInEnabled)
+{
+	bIsEnabled = bInEnabled;
+}
+
 void UCompositeLayerBase::AddChildPasses(UE::CompositeCore::FPassInputDecl& InBasePassInput, FTraversalContext& InContext, FSceneRenderingBulkObjectAllocator& InFrameAllocator, TArrayView<const TObjectPtr<UCompositePassBase>> InPasses) const
 {
 	// Iterate backwards so that lower passes in the UI are executed first
diff --git a/Engine/Plugins/Experimental/Compositing/Composite/Source/Composite/Private/Layers/CompositeLayerPlate.cpp b/Engine/Plugins/Experimental/Compositing/Composite/Source/Composite/Private/Layers/CompositeLayerPlate.cpp
--- a/Engine/Plugins/Experimental/Compositing/Composite/Source/Composite/Private/Layers/CompositeLayerPlate.cpp
+++ b/Engine/Plugins/Experimental/Compositing/Composite/Source/Composite/Private/Layers/CompositeLayerPlate.cpp
@@ -94,6 +94,13 @@ UCompositeLayerPlate::UCompositeLayerPlate(const FObjectInitializer& ObjectIniti
 
 UCompositeLayerPlate::~UCompositeLayerPlate() = default;
 
+void UCompositeLayerPlate::SetEnabled(bool bInEnabled)
+{
+	Super::SetEnabled(bInEnabled);
+
+	HandleEnabledStateChange(IsEnabled(), UWorld::GetSubsystem<UCompositeCoreSubsystem>(GetWorld()));
+}
+
 void UCompositeLayerPlate::Tick(float DeltaTime)
 {
 	if(!bIsEnabled)
@@ -295,7 +302,7 @@ void UCompositeLayerPlate::PostEditChangeProperty(FPropertyChangedEvent& Propert
 
 	if (PropertyName == GET_MEMBER_NAME_CHECKED(ThisClass, bIsEnabled))
 	{
-		HandleEnabledStateChange(bIsEnabled, UWorld::GetSubsystem<UCompositeCoreSubsystem>(GetWorld()));
+		SetEnabled(bIsEnabled);
 	}
 	else if (PropertyName == GET_MEMBER_NAME_CHECKED(ThisClass, CompositeMeshes))
 	{
diff --git a/Engine/Plugins/Experimental/Compositing/Composite/Source/Composite/Private/Layers/CompositeLayerSceneCapture.cpp b/Engine/Plugins/Experimental/Compositing/Composite/Source/Composite/Private/Layers/CompositeLayerSceneCapture.cpp
--- a/Engine/Plugins/Experimental/Compositing/Composite/Source/Composite/Private/Layers/CompositeLayerSceneCapture.cpp
+++ b/Engine/Plugins/Experimental/Compositing/Composite/Source/Composite/Private/Layers/CompositeLayerSceneCapture.cpp
@@ -19,6 +19,17 @@ UCompositeLayerSceneCapture::UCompositeLayerSceneCapture(const FObjectInitialize
 
 UCompositeLayerSceneCapture::~UCompositeLayerSceneCapture() = default;
 
+void UCompositeLayerSceneCapture::SetEnabled(bool bInEnabled)
+{
+	Super::SetEnabled(bInEnabled);
+
+	ACompositeActor* CompositeActor = GetTypedOuter<ACompositeActor>();
+	if (IsValid(CompositeActor))
+	{
+		CompositeActor->ActivateSceneCaptures(this, IsEnabled());
+	}
+}
+
 void UCompositeLayerSceneCapture::BeginDestroy()
 {
 	Super::BeginDestroy();
@@ -38,11 +49,7 @@ void UCompositeLayerSceneCapture::PostEditChangeProperty(FPropertyChangedEvent&
 	
 	if (PropertyName == GET_MEMBER_NAME_CHECKED(ThisClass, bIsEnabled))
 	{
-		ACompositeActor* CompositeActor = GetTypedOuter<ACompositeActor>();
-		if (IsValid(CompositeActor))
-		{
-			CompositeActor->ActivateSceneCaptures(this, bIsEnabled);
-		}
+		SetEnabled(bIsEnabled);
 	}
 	else if (PropertyName == GET_MEMBER_NAME_CHECKED(ThisClass, Actors))
 	{
diff --git a/Engine/Plugins/Experimental/Compositing/Composite/Source/Composite/Private/Layers/CompositeLayerShadowReflection.cpp b/Engine/Plugins/Experimental/Compositing/Composite/Source/Composite/Private/Layers/CompositeLayerShadowReflection.cpp
--- a/Engine/Plugins/Experimental/Compositing/Composite/Source/Composite/Private/Layers/CompositeLayerShadowReflection.cpp
+++ b/Engine/Plugins/Experimental/Compositing/Composite/Source/Composite/Private/Layers/CompositeLayerShadowReflection.cpp
@@ -16,6 +16,17 @@ UCompositeLayerShadowReflection::UCompositeLayerShadowReflection(const FObjectIn
 
 UCompositeLayerShadowReflection::~UCompositeLayerShadowReflection() = default;
 
+void UCompositeLayerShadowReflection::SetEnabled(bool bInEnabled)
+{
+	Super::SetEnabled(bInEnabled);
+
+	ACompositeActor* CompositeActor = GetTypedOuter<ACompositeActor>();
+	if (IsValid(CompositeActor))
+	{
+		CompositeActor->ActivateSceneCaptures(this, IsEnabled());
+	}
+}
+
 #if WITH_EDITOR
 void UCompositeLayerShadowReflection::PostEditChangeProperty(FPropertyChangedEvent& PropertyChangedEvent)
 {
@@ -25,11 +36,7 @@ void UCompositeLayerShadowReflection::PostEditChangeProperty(FPropertyChangedEve
 
 	if (PropertyName == GET_MEMBER_NAME_CHECKED(ThisClass, bIsEnabled))
 	{
-		ACompositeActor* CompositeActor = GetTypedOuter<ACompositeActor>();
-		if (IsValid(CompositeActor))
-		{
-			CompositeActor->ActivateSceneCaptures(this, bIsEnabled);
-		}
+		SetEnabled(bIsEnabled);
 	}
 	else if (PropertyName == GET_MEMBER_NAME_CHECKED(ThisClass, Actors))
 	{
diff --git a/Engine/Plugins/Experimental/Compositing/Composite/Source/Composite/Private/Passes/CompositePassBase.cpp b/Engine/Plugins/Experimental/Compositing/Composite/Source/Composite/Private/Passes/CompositePassBase.cpp
--- a/Engine/Plugins/Experimental/Compositing/Composite/Source/Composite/Private/Passes/CompositePassBase.cpp
+++ b/Engine/Plugins/Experimental/Compositing/Composite/Source/Composite/Private/Passes/CompositePassBase.cpp
@@ -12,6 +12,16 @@ UCompositePassBase::UCompositePassBase(const FObjectInitializer& ObjectInitializ
 
 UCompositePassBase::~UCompositePassBase() = default;
 
+bool UCompositePassBase::IsEnabled() const
+{
+	return bIsEnabled;
+}
+
+void UCompositePassBase::SetEnabled(bool bInEnabled)
+{
+	bIsEnabled = bInEnabled;
+}
+
 #if WITH_EDITOR
 bool UCompositePassBase::CanEditChange(const FProperty* InProperty) const
 {
diff --git a/Engine/Plugins/Experimental/Compositing/Composite/Source/Composite/Private/Passes/CompositePassCenteredScale.cpp b/Engine/Plugins/Experimental/Compositing/Composite/Source/Composite/Private/Passes/CompositePassCenteredScale.cpp
--- a/Engine/Plugins/Experimental/Compositing/Composite/Source/Composite/Private/Passes/CompositePassCenteredScale.cpp
+++ b/Engine/Plugins/Experimental/Compositing/Composite/Source/Composite/Private/Passes/CompositePassCenteredScale.cpp
@@ -112,7 +112,7 @@ bool UCompositePassCenteredScale::GetProxy(const UE::CompositeCore::FPassInputDe
 
 bool UCompositePassCenteredScale::IsActive() const
 {
-	return bIsEnabled && !CalculateScale().Equals(FVector2f::One(), UE_SMALL_NUMBER);
+	return IsEnabled() && !CalculateScale().Equals(FVector2f::One(), UE_SMALL_NUMBER);
 }
 
 FVector2f UCompositePassCenteredScale::CalculateScale() const
diff --git a/Engine/Plugins/Experimental/Compositing/Composite/Source/Composite/Public/Layers/CompositeLayerBase.h b/Engine/Plugins/Experimental/Compositing/Composite/Source/Composite/Public/Layers/CompositeLayerBase.h
--- a/Engine/Plugins/Experimental/Compositing/Composite/Source/Composite/Public/Layers/CompositeLayerBase.h
+++ b/Engine/Plugins/Experimental/Compositing/Composite/Source/Composite/Public/Layers/CompositeLayerBase.h
@@ -55,11 +55,15 @@ public:
 	/** Getter function to override, returning pass layer proxies to be passed to the render thread. (Proxy objects should be allocated from the provided allocator.) */
 	UE_API virtual bool GetProxy(FTraversalContext& InContext, FSceneRenderingBulkObjectAllocator& InFrameAllocator, FCompositeCorePassProxy*& OutProxy) const { return false; };
 
-public:
-	/** Whether or not the pass is enabled. */
-	UPROPERTY(EditAnywhere, BlueprintReadWrite, Category = "", meta = (DisplayPriority = "2"))
-	bool bIsEnabled = true;
+	/** Get the enabled state. */
+	UFUNCTION(BlueprintGetter)
+	UE_API virtual bool IsEnabled() const;
+
+	/** Set the enabled state. */
+	UFUNCTION(BlueprintSetter)
+	UE_API virtual void SetEnabled(bool bInEnabled);
 
+public:
 	/** Whether or not the pass is solo. */
 	UPROPERTY()
 	bool bIsSolo = false;
@@ -75,6 +79,9 @@ public:
 	ECompositeCoreMergeOp Operation = ECompositeCoreMergeOp::Over;
 
 protected:
+	/** Whether or not the pass is enabled. */
+	UPROPERTY(EditAnywhere, BlueprintGetter = IsEnabled, BlueprintSetter = SetEnabled, Category = "", meta = (DisplayPriority = "2"))
+	bool bIsEnabled = true;
 
 	/** Convenience function to register post-processing child passes. */
 	void AddChildPasses(UE::CompositeCore::FPassInputDecl& InBasePassInput, FTraversalContext& InContext, FSceneRenderingBulkObjectAllocator& InFrameAllocator, TArrayView<const TObjectPtr<UCompositePassBase>> InPasses) const;
diff --git a/Engine/Plugins/Experimental/Compositing/Composite/Source/Composite/Public/Layers/CompositeLayerPlate.h b/Engine/Plugins/Experimental/Compositing/Composite/Source/Composite/Public/Layers/CompositeLayerPlate.h
--- a/Engine/Plugins/Experimental/Compositing/Composite/Source/Composite/Public/Layers/CompositeLayerPlate.h
+++ b/Engine/Plugins/Experimental/Compositing/Composite/Source/Composite/Public/Layers/CompositeLayerPlate.h
@@ -45,6 +45,9 @@ public:
 	/** Destructor. */
 	UE_API virtual ~UCompositeLayerPlate();
 
+	/** Set the enabled state. */
+	UE_API virtual void SetEnabled(bool bInEnabled) override;
+
 	//~ Begin FTickableGameObject interface
 	UE_API virtual void Tick(float DeltaTime) override;
 	
diff --git a/Engine/Plugins/Experimental/Compositing/Composite/Source/Composite/Public/Layers/CompositeLayerSceneCapture.h b/Engine/Plugins/Experimental/Compositing/Composite/Source/Composite/Public/Layers/CompositeLayerSceneCapture.h
--- a/Engine/Plugins/Experimental/Compositing/Composite/Source/Composite/Public/Layers/CompositeLayerSceneCapture.h
+++ b/Engine/Plugins/Experimental/Compositing/Composite/Source/Composite/Public/Layers/CompositeLayerSceneCapture.h
@@ -25,6 +25,9 @@ public:
 	/** Destructor. */
 	UE_API virtual ~UCompositeLayerSceneCapture();
 
+	/** Set the enabled state. */
+	UE_API virtual void SetEnabled(bool bInEnabled) override;
+
 	//~ Begin UObject Interface
 	UE_API virtual void BeginDestroy() override;
 	//~ End UObject Interface
diff --git a/Engine/Plugins/Experimental/Compositing/Composite/Source/Composite/Public/Layers/CompositeLayerShadowReflection.h b/Engine/Plugins/Experimental/Compositing/Composite/Source/Composite/Public/Layers/CompositeLayerShadowReflection.h
--- a/Engine/Plugins/Experimental/Compositing/Composite/Source/Composite/Public/Layers/CompositeLayerShadowReflection.h
+++ b/Engine/Plugins/Experimental/Compositing/Composite/Source/Composite/Public/Layers/CompositeLayerShadowReflection.h
@@ -37,6 +37,9 @@ public:
 	/** Destructor. */
 	UE_API virtual ~UCompositeLayerShadowReflection();
 
+	/** Set the enabled state. */
+	UE_API virtual void SetEnabled(bool bInEnabled) override;
+
 #if WITH_EDITOR
 	UE_API virtual void PostEditChangeProperty(FPropertyChangedEvent& PropertyChangedEvent) override;
 #endif
diff --git a/Engine/Plugins/Experimental/Compositing/Composite/Source/Composite/Public/Passes/CompositePassBase.h b/Engine/Plugins/Experimental/Compositing/Composite/Source/Composite/Public/Passes/CompositePassBase.h
--- a/Engine/Plugins/Experimental/Compositing/Composite/Source/Composite/Public/Passes/CompositePassBase.h
+++ b/Engine/Plugins/Experimental/Compositing/Composite/Source/Composite/Public/Passes/CompositePassBase.h
@@ -32,14 +32,23 @@ public:
 	/** Getter function to override, returning pass proxies to be passed to the render thread. (Proxy objects should be allocated from the provided allocator.) */
 	UE_API virtual bool GetProxy(const UE::CompositeCore::FPassInputDecl& InputDecl, FSceneRenderingBulkObjectAllocator& InFrameAllocator, FCompositeCorePassProxy*& OutProxy) const { return false; };
 
+	/** Get the enabled state. */
+	UFUNCTION(BlueprintGetter)
+	UE_API virtual bool IsEnabled() const;
+
+	/** Set the enabled state. */
+	UFUNCTION(BlueprintSetter)
+	UE_API virtual void SetEnabled(bool bInEnabled);
+
 	/** Returns true if the pass is returns a valid proxy. */
-	virtual bool IsActive() const { return bIsEnabled; }
+	virtual bool IsActive() const { return IsEnabled(); }
 
 	/** Returns true if the pass (may) need post-processing scene textures. */
 	virtual bool NeedsSceneTextures() const { return false; }
 
+protected:
 	/** Whether or not the pass is active. */
-	UPROPERTY(EditAnywhere, BlueprintReadWrite, Category = "", meta = (DisplayPriority = "1"))
+	UPROPERTY(EditAnywhere, BlueprintGetter = IsEnabled, BlueprintSetter = SetEnabled, Category = "", meta = (DisplayPriority = "1"))
 	bool bIsEnabled = true;
 };
 
diff --git a/Engine/Plugins/Experimental/Compositing/Composite/Source/CompositeEditor/Private/UI/SCompositePanelLayerTree.cpp b/Engine/Plugins/Experimental/Compositing/Composite/Source/CompositeEditor/Private/UI/SCompositePanelLayerTree.cpp
--- a/Engine/Plugins/Experimental/Compositing/Composite/Source/CompositeEditor/Private/UI/SCompositePanelLayerTree.cpp
+++ b/Engine/Plugins/Experimental/Compositing/Composite/Source/CompositeEditor/Private/UI/SCompositePanelLayerTree.cpp
@@ -414,7 +414,7 @@ private:
 		{
 			if (UCompositeLayerBase* CompositeLayer = TreeItem->GetCompositeLayer())
 			{
-				return CompositeLayer->bIsEnabled ? ECheckBoxState::Checked : ECheckBoxState::Unchecked;
+				return CompositeLayer->IsEnabled() ? ECheckBoxState::Checked : ECheckBoxState::Unchecked;
 			}
 			
 			// In the case where there is a valid layer index, but the layer itself is null
@@ -439,8 +439,7 @@ private:
 				FScopedTransaction SetEnabledTransaction(LOCTEXT("SetEnabledTransaction", "Set Enabled"));
 				CompositeLayer->Modify();
 
-				// TODO: Should convert to a setter that can ensure ActivateSceneCapture is called when re-enabling
-				CompositeLayer->bIsEnabled = InCheckBoxState == ECheckBoxState::Checked;
+				CompositeLayer->SetEnabled(InCheckBoxState == ECheckBoxState::Checked);
 				return;
 			}
 
@@ -1188,7 +1187,7 @@ ECheckBoxState SCompositePanelLayerTree::GetGlobalEnabledCheckState() const
 
 			if (UCompositeLayerBase* CompositeLayer = ChildTreeItem->GetCompositeLayer())
 			{
-				if (CompositeLayer->bIsEnabled)
+				if (CompositeLayer->IsEnabled())
 				{
 					bAnyEnabled = true;
 				}
@@ -1241,7 +1240,7 @@ void SCompositePanelLayerTree::OnGlobalEnabledCheckStateChanged(ECheckBoxState C
 			if (UCompositeLayerBase* CompositeLayer = ChildTreeItem->GetCompositeLayer())
 			{
 				CompositeLayer->Modify();
-				CompositeLayer->bIsEnabled = CheckBoxState == ECheckBoxState::Checked;
+				CompositeLayer->SetEnabled(CheckBoxState == ECheckBoxState::Checked);
 			}
 		}
 	}
@@ -1916,7 +1915,7 @@ void SCompositePanelLayerTree::EnableSelectedItems()
 		{
 			if (UCompositeLayerBase* Layer = Item->GetCompositeLayer())
 			{
-				if (Layer->bIsEnabled)
+				if (Layer->IsEnabled())
 				{
 					bAnyEnabled = true;
 				}
@@ -1947,7 +1946,7 @@ void SCompositePanelLayerTree::EnableSelectedItems()
 			if (UCompositeLayerBase* Layer = Item->GetCompositeLayer())
 			{
 				Layer->Modify();
-				Layer->bIsEnabled = bSetEnabled;
+				Layer->SetEnabled(bSetEnabled);
 			}
 		}
 		else
diff --git a/Engine/Plugins/Experimental/Compositing/Composite/Source/CompositeEditor/Private/UI/SCompositePassTree.cpp b/Engine/Plugins/Experimental/Compositing/Composite/Source/CompositeEditor/Private/UI/SCompositePassTree.cpp
--- a/Engine/Plugins/Experimental/Compositing/Composite/Source/CompositeEditor/Private/UI/SCompositePassTree.cpp
+++ b/Engine/Plugins/Experimental/Compositing/Composite/Source/CompositeEditor/Private/UI/SCompositePassTree.cpp
@@ -586,7 +586,7 @@ private:
 		{
 			if (UCompositePassBase* Pass = TreeItem->GetPass())
 			{
-				return Pass->bIsEnabled ? ECheckBoxState::Checked : ECheckBoxState::Unchecked;
+				return Pass->IsEnabled() ? ECheckBoxState::Checked : ECheckBoxState::Unchecked;
 			}
 			
 			// In the case where there is a valid pass index, but the pass itself is null
@@ -600,7 +600,7 @@ private:
 		bool bAnyDisabled = false;
 		for (const TObjectPtr<UCompositePassBase>& Pass : AllPassesOfType)
 		{
-			if (Pass->bIsEnabled)
+			if (Pass->IsEnabled())
 			{
 				bAnyEnabled = true;
 			}
@@ -635,7 +635,7 @@ private:
 				FScopedTransaction SetEnabledTransaction(LOCTEXT("SetEnabledTransaction", "Set Enabled"));
 				Pass->Modify();
 
-				Pass->bIsEnabled = InCheckBoxState == ECheckBoxState::Checked;
+				Pass->SetEnabled(InCheckBoxState == ECheckBoxState::Checked);
 				return;
 			}
 
@@ -654,7 +654,7 @@ private:
 		for (const TObjectPtr<UCompositePassBase>& Pass : AllPassesOfType)
 		{
 			Pass->Modify();
-			Pass->bIsEnabled = InCheckBoxState == ECheckBoxState::Checked;
+			Pass->SetEnabled(InCheckBoxState == ECheckBoxState::Checked);
 		}
 	}
 	
@@ -1126,7 +1126,7 @@ ECheckBoxState SCompositePassTree::GetGlobalEnabledState() const
 		
 		for (const TObjectPtr<UCompositePassBase>& Pass : PassList)
 		{
-			if (Pass->bIsEnabled)
+			if (Pass->IsEnabled())
 			{
 				bAnyPassesEnabled = true;
 			}
@@ -1169,7 +1169,7 @@ void SCompositePassTree::OnGlobalEnabledStateChanged(ECheckBoxState CheckBoxStat
 			}
 
 			Pass->Modify();
-			Pass->bIsEnabled = CheckBoxState == ECheckBoxState::Checked;
+			Pass->SetEnabled(CheckBoxState == ECheckBoxState::Checked);
 		}
 	}
 }
@@ -1687,7 +1687,7 @@ void SCompositePassTree::EnableSelectedItems()
 		{
 			if (UCompositePassBase* Pass = Item->GetPass())
 			{
-				if (Pass->bIsEnabled)
+				if (Pass->IsEnabled())
 				{
 					bAnyEnabled = true;
 				}
@@ -1701,7 +1701,7 @@ void SCompositePassTree::EnableSelectedItems()
 				{
 					if (UCompositePassBase* Pass = ChildItem->GetPass())
 					{
-						if (Pass->bIsEnabled)
+						if (Pass->IsEnabled())
 						{
 							bAnyEnabled = true;
 						}
@@ -1723,7 +1723,7 @@ void SCompositePassTree::EnableSelectedItems()
 			if (UCompositePassBase* Pass = Item->GetPass())
 			{
 				Pass->Modify();
-				Pass->bIsEnabled = bSetEnabled;
+				Pass->SetEnabled(bSetEnabled);
 			}
 		}
 		else
@@ -1735,7 +1735,7 @@ void SCompositePassTree::EnableSelectedItems()
 					if (UCompositePassBase* Pass = ChildItem->GetPass())
 					{
 						Pass->Modify();
-						Pass->bIsEnabled = bSetEnabled;
+						Pass->SetEnabled(bSetEnabled);
 					}
 				}
 			}
