Composite: Add enabled condition on add layer button.

diff --git a/Engine/Plugins/Experimental/Compositing/Composite/Source/CompositeEditor/Private/UI/SCompositePanelLayerTree.cpp b/Engine/Plugins/Experimental/Compositing/Composite/Source/CompositeEditor/Private/UI/SCompositePanelLayerTree.cpp
--- a/Engine/Plugins/Experimental/Compositing/Composite/Source/CompositeEditor/Private/UI/SCompositePanelLayerTree.cpp
+++ b/Engine/Plugins/Experimental/Compositing/Composite/Source/CompositeEditor/Private/UI/SCompositePanelLayerTree.cpp
@@ -47,6 +47,7 @@ public:
 	DECLARE_DELEGATE_OneParam(FOnLayerAdded, const UClass*)
 	
 	SLATE_BEGIN_ARGS(SCompositePanelLayerTreeToolbar) {}
+		SLATE_ATTRIBUTE(bool, CanAddLayer)
 		SLATE_ATTRIBUTE(bool, CanPilotCamera)
 		SLATE_ATTRIBUTE(bool, IsPilotingCamera)
 		SLATE_EVENT(FOnPilotCameraToggled, OnPilotCameraToggled)
@@ -80,6 +81,7 @@ public:
 					.Icon(FAppStyle::Get().GetBrush("Icons.Plus"))
 					.Text(LOCTEXT("AddMediaButtonLabel", "Add"))
 					.ToolTipText(LOCTEXT("AddMediaButtonToolTip", "Add new composite layer to the selected composite actor"))
+					.IsEnabled(InArgs._CanAddLayer)
 					.OnGetMenuContent(this, &SCompositePanelLayerTreeToolbar::GetAddMenuContent)
 				]
 
@@ -749,6 +751,7 @@ void SCompositePanelLayerTree::Construct(const FArguments& InArgs)
 		.Padding(8.0f, 8.0f, 8.0f, 4.0f)
 		[
 			SAssignNew(Toolbar, SCompositePanelLayerTreeToolbar)
+			.CanAddLayer(this, &SCompositePanelLayerTree::CanAddLayer)
 			.CanPilotCamera(this, &SCompositePanelLayerTree::CanPilotCamera)
 			.IsPilotingCamera(this, &SCompositePanelLayerTree::IsPilotingCamera)
 			.OnPilotCameraToggled(this, &SCompositePanelLayerTree::OnPilotCameraToggled)
@@ -1080,6 +1083,11 @@ void SCompositePanelLayerTree::RefreshAndExpandTreeView()
 	}
 }
 
+bool SCompositePanelLayerTree::CanAddLayer() const
+{
+	return TreeView.IsValid() && !TreeView->GetSelectedItems().IsEmpty();
+}
+
 bool SCompositePanelLayerTree::CanPilotCamera() const
 {
 	TArray<FCompositeActorTreeItemPtr> SelectedItems = TreeView.IsValid() ? TreeView->GetSelectedItems() : TArray<FCompositeActorTreeItemPtr>();
diff --git a/Engine/Plugins/Experimental/Compositing/Composite/Source/CompositeEditor/Private/UI/SCompositePanelLayerTree.h b/Engine/Plugins/Experimental/Compositing/Composite/Source/CompositeEditor/Private/UI/SCompositePanelLayerTree.h
--- a/Engine/Plugins/Experimental/Compositing/Composite/Source/CompositeEditor/Private/UI/SCompositePanelLayerTree.h
+++ b/Engine/Plugins/Experimental/Compositing/Composite/Source/CompositeEditor/Private/UI/SCompositePanelLayerTree.h
@@ -84,6 +84,9 @@ private:
 	/** Refreshes the tree view and expands all tree view items */
 	void RefreshAndExpandTreeView();
 
+	/** Gets whether it is possible to add a layer based on active selection or lack thereof. */
+	bool CanAddLayer() const;
+
 	/** Gets whether the currently selected composite actor has a camera that can be piloted */
 	bool CanPilotCamera() const;
 
