Composure: Added custom actor list and picker that allows easy multi-select and drag and drop for the Shadow/Reflection and Plate actor content lists

diff --git a/Engine/Plugins/Experimental/Compositing/Composite/Source/Composite/Public/Layers/CompositeLayerPlate.h b/Engine/Plugins/Experimental/Compositing/Composite/Source/Composite/Public/Layers/CompositeLayerPlate.h
--- a/Engine/Plugins/Experimental/Compositing/Composite/Source/Composite/Public/Layers/CompositeLayerPlate.h
+++ b/Engine/Plugins/Experimental/Compositing/Composite/Source/Composite/Public/Layers/CompositeLayerPlate.h
@@ -205,6 +205,8 @@ private:
 
 	/** Cached number of valid scene passes. */
 	int32 CachedValidScenePasses = 0;
+
+	friend class FCompositeLayerPlateCustomization;
 };
 
 #undef UE_API
diff --git a/Engine/Plugins/Experimental/Compositing/Composite/Source/Composite/Public/Layers/CompositeLayerShadowReflection.h b/Engine/Plugins/Experimental/Compositing/Composite/Source/Composite/Public/Layers/CompositeLayerShadowReflection.h
--- a/Engine/Plugins/Experimental/Compositing/Composite/Source/Composite/Public/Layers/CompositeLayerShadowReflection.h
+++ b/Engine/Plugins/Experimental/Compositing/Composite/Source/Composite/Public/Layers/CompositeLayerShadowReflection.h
@@ -72,6 +72,8 @@ private:
 	/** List of reflection & shadow caster actors. */
 	UPROPERTY(EditAnywhere, BlueprintGetter = GetActors, BlueprintSetter = SetActors, Category = "")
 	TArray<TObjectPtr<AActor>> Actors;
+
+	friend class FCompositeLayerShadowReflectionCustomization;
 };
 
 #undef UE_API
diff --git a/Engine/Plugins/Experimental/Compositing/Composite/Source/CompositeEditor/CompositeEditor.Build.cs b/Engine/Plugins/Experimental/Compositing/Composite/Source/CompositeEditor/CompositeEditor.Build.cs
--- a/Engine/Plugins/Experimental/Compositing/Composite/Source/CompositeEditor/CompositeEditor.Build.cs
+++ b/Engine/Plugins/Experimental/Compositing/Composite/Source/CompositeEditor/CompositeEditor.Build.cs
@@ -22,10 +22,12 @@ public class CompositeEditor : ModuleRules
 				"ApplicationCore",
 				"CompositeCore",
 				"CoreUObject",
+				"EditorWidgets",
 				"Engine",
 				"InputCore",
 				"LevelEditor",
 				"UnrealEd",
+				"SceneOutliner",
                 "Slate",
                 "SlateCore",
                 "ToolMenus",
diff --git a/Engine/Plugins/Experimental/Compositing/Composite/Source/CompositeEditor/Private/CompositeEditorCommands.cpp b/Engine/Plugins/Experimental/Compositing/Composite/Source/CompositeEditor/Private/CompositeEditorCommands.cpp
--- a/Engine/Plugins/Experimental/Compositing/Composite/Source/CompositeEditor/Private/CompositeEditorCommands.cpp
+++ b/Engine/Plugins/Experimental/Compositing/Composite/Source/CompositeEditor/Private/CompositeEditorCommands.cpp
@@ -9,6 +9,7 @@
 void FCompositeEditorCommands::RegisterCommands()
 {
 	UI_COMMAND(Enable, "Toggle Enable", "Toggle Enable State", EUserInterfaceActionType::Button, FInputChord(EModifierKey::None, EKeys::D));
+	UI_COMMAND(RemoveActor, "Remove Actor", "Removes selected actor(s) from the list", EUserInterfaceActionType::Button, FInputChord(EModifierKey::None, EKeys::Delete));
 }
 
 #undef LOCTEXT_NAMESPACE
diff --git a/Engine/Plugins/Experimental/Compositing/Composite/Source/CompositeEditor/Private/CompositeEditorCommands.h b/Engine/Plugins/Experimental/Compositing/Composite/Source/CompositeEditor/Private/CompositeEditorCommands.h
--- a/Engine/Plugins/Experimental/Compositing/Composite/Source/CompositeEditor/Private/CompositeEditorCommands.h
+++ b/Engine/Plugins/Experimental/Compositing/Composite/Source/CompositeEditor/Private/CompositeEditorCommands.h
@@ -6,7 +6,7 @@
 #include "Styling/AppStyle.h"
 
 /**
- * 
+ * UI commands for actions that are specific to Composite editor
  */
 class FCompositeEditorCommands : public TCommands<FCompositeEditorCommands>
 {
@@ -23,4 +23,5 @@ public:
 
 public:
 	TSharedPtr<FUICommandInfo> Enable;
+	TSharedPtr<FUICommandInfo> RemoveActor;
 };
diff --git a/Engine/Plugins/Experimental/Compositing/Composite/Source/CompositeEditor/Private/CompositeEditorModule.cpp b/Engine/Plugins/Experimental/Compositing/Composite/Source/CompositeEditor/Private/CompositeEditorModule.cpp
--- a/Engine/Plugins/Experimental/Compositing/Composite/Source/CompositeEditor/Private/CompositeEditorModule.cpp
+++ b/Engine/Plugins/Experimental/Compositing/Composite/Source/CompositeEditor/Private/CompositeEditorModule.cpp
@@ -14,7 +14,9 @@
 #include "CompositeEditorStyle.h"
 #include "PropertyEditorModule.h"
 #include "Customizations/CompositeLayerPlateCustomization.h"
+#include "Customizations/CompositeLayerShadowReflectionCustomization.h"
 #include "Layers/CompositeLayerPlate.h"
+#include "Layers/CompositeLayerShadowReflection.h"
 #include "UI/SCompositeEditorPanel.h"
 
 #define LOCTEXT_NAMESPACE "CompositeEditorModule"
@@ -116,12 +118,23 @@ void FCompositeEditorModule::RegisterCustomizations()
 		UCompositeLayerPlate::StaticClass()->GetFName(),
 		FOnGetDetailCustomizationInstance::CreateStatic(&FCompositeLayerPlateCustomization::MakeInstance)
 	);
+
+	PropertyModule.RegisterCustomClassLayout(
+		UCompositeLayerShadowReflection::StaticClass()->GetFName(),
+		FOnGetDetailCustomizationInstance::CreateStatic(&FCompositeLayerShadowReflectionCustomization::MakeInstance)
+	);
 }
 
 void FCompositeEditorModule::UnregisterCustomizations()
 {
-	FPropertyEditorModule& PropertyModule = FModuleManager::LoadModuleChecked<FPropertyEditorModule>("PropertyEditor");
-	PropertyModule.UnregisterCustomClassLayout(UCompositeLayerPlate::StaticClass()->GetFName());
+	if (UObjectInitialized() && !IsEngineExitRequested())
+	{
+		if (FPropertyEditorModule* PropertyModule = FModuleManager::Get().GetModulePtr<FPropertyEditorModule>("PropertyEditor"))
+		{
+			PropertyModule->UnregisterCustomClassLayout(UCompositeLayerPlate::StaticClass()->GetFName());
+			PropertyModule->UnregisterCustomClassLayout(UCompositeLayerShadowReflection::StaticClass()->GetFName());
+		}
+	}
 }
 
 IMPLEMENT_MODULE(FCompositeEditorModule, CompositeEditor)
diff --git a/Engine/Plugins/Experimental/Compositing/Composite/Source/CompositeEditor/Private/Customizations/CompositeLayerPlateCustomization.cpp b/Engine/Plugins/Experimental/Compositing/Composite/Source/CompositeEditor/Private/Customizations/CompositeLayerPlateCustomization.cpp
--- a/Engine/Plugins/Experimental/Compositing/Composite/Source/CompositeEditor/Private/Customizations/CompositeLayerPlateCustomization.cpp
+++ b/Engine/Plugins/Experimental/Compositing/Composite/Source/CompositeEditor/Private/Customizations/CompositeLayerPlateCustomization.cpp
@@ -2,24 +2,27 @@
 
 #include "CompositeLayerPlateCustomization.h"
 
+#include "CompositeActor.h"
+#include "CompositeEditorStyle.h"
+#include "CompositeMeshActor.h"
 #include "DetailCategoryBuilder.h"
 #include "DetailLayoutBuilder.h"
 #include "DetailWidgetRow.h"
+#include "Editor.h"
 #include "IDetailGroup.h"
 #include "IPropertyUtilities.h"
+#include "ScopedTransaction.h"
+#include "UnrealEdGlobals.h"
+#include "Camera/CameraComponent.h"
+#include "Editor/UnrealEdEngine.h"
+#include "Framework/MultiBox/MultiBoxBuilder.h"
 #include "Layers/CompositeLayerPlate.h"
+#include "UI/SCompositeActorPickerTable.h"
 #include "UI/SCompositePlatePassPanel.h"
 #include "Widgets/Text/STextBlock.h"
 
 #define LOCTEXT_NAMESPACE "FCompositeLayerPlateCustomization"
 
-/** Temporary cvar to ensure that the old passes lists can still be exposed in case something is missing or incorrect with the new pass table */
-TAutoConsoleVariable<bool> CVarDisablePassTableCustomization(
-	TEXT("Composure.DisablePassTableCustomization"),
-	/*DefaultValue*/false,
-	TEXT("If `true`, composite plates do not show the new pass table in their details panel, instead showing the default passes lists")
-);
-
 TSharedRef<IDetailCustomization> FCompositeLayerPlateCustomization::MakeInstance()
 {
 	return MakeShared<FCompositeLayerPlateCustomization>();
@@ -33,14 +36,9 @@ void FCompositeLayerPlateCustomization::CustomizeDetails(const TSharedPtr<IDetai
 
 void FCompositeLayerPlateCustomization::CustomizeDetails(IDetailLayoutBuilder& DetailLayout)
 {
-	// TODO: Once we are confident in the pass table UI, this can be removed, but for now allow fallback to non-customized details panel for testing/debug
-	if (CVarDisablePassTableCustomization.GetValueOnAnyThread())
-	{
-		return;
-	}
-	
 	TArray<TWeakObjectPtr<UCompositeLayerPlate>> Objects = DetailLayout.GetObjectsOfTypeBeingCustomized<UCompositeLayerPlate>();
 
+	DetailLayout.HideProperty(GET_MEMBER_NAME_CHECKED(UCompositeLayerPlate, CompositeMeshes));
 	DetailLayout.HideProperty(GET_MEMBER_NAME_CHECKED(UCompositeLayerPlate, MediaPasses));
 	DetailLayout.HideProperty(GET_MEMBER_NAME_CHECKED(UCompositeLayerPlate, LayerPasses));
 	
@@ -60,6 +58,17 @@ void FCompositeLayerPlateCustomization::CustomizeDetails(IDetailLayoutBuilder& D
 	if (Objects.Num() == 1 && Objects[0].IsValid())
 	{
 		UCompositeLayerPlate* Plate = Cast<UCompositeLayerPlate>(Objects[0].Get());
+
+		IDetailGroup& ActorListGroup = PlateCategory.AddGroup("CompositeMeshContent", LOCTEXT("CompositeMeshGroupName", "Composite Mesh Content"), false, true);
+		FCompositeActorPickerListRef ActorListRef(Plate, GET_MEMBER_NAME_CHECKED(UCompositeLayerPlate, CompositeMeshes), &Plate->CompositeMeshes);
+		
+		ActorListGroup.AddWidgetRow()
+		.WholeRowContent()
+		[
+			SNew(SCompositeActorPickerTable, ActorListRef)
+			.OnExtendAddMenu(this, &FCompositeLayerPlateCustomization::OnExtendCompositeMeshAddMenu)
+			.OnLayoutSizeChanged(this, &FCompositeLayerPlateCustomization::OnLayoutSizeChanged)
+		];
 		
 		IDetailGroup& PassesGroup = PlateCategory.AddGroup("Passes", LOCTEXT("PassesGroupName", "Passes"), false, true);
 
@@ -67,14 +76,7 @@ void FCompositeLayerPlateCustomization::CustomizeDetails(IDetailLayoutBuilder& D
 		.WholeRowContent()
 		[
 			SNew(SCompositePlatePassPanel, Plate)
-			.OnLayoutSizeChanged_Lambda([this]
-			{
-				TSharedPtr<IDetailLayoutBuilder> PinnedDetailBuilder = CachedDetailBuilder.Pin();
-				if (PinnedDetailBuilder.IsValid())
-				{
-					PinnedDetailBuilder->GetPropertyUtilities()->RequestRefresh();
-				}
-			})
+			.OnLayoutSizeChanged(this, &FCompositeLayerPlateCustomization::OnLayoutSizeChanged)
 		];
 	}  
 	else
@@ -96,4 +98,122 @@ void FCompositeLayerPlateCustomization::CustomizeDetails(IDetailLayoutBuilder& D
 	}
 }
 
+void FCompositeLayerPlateCustomization::OnExtendCompositeMeshAddMenu(FMenuBuilder& MenuBuilder)
+{
+	MenuBuilder.AddMenuEntry(
+		LOCTEXT("PlaceCompositeMeshActorEntryLabel", "Place Composite Mesh Actor"),
+		LOCTEXT("PlaceCompositeMeshActorEntryToolTip", "Place a composite mesh actor at the appropriate position in the level"),
+		FSlateIcon(FCompositeEditorStyle::Get().GetStyleSetName(), "CompositeEditor.Composure"),
+		FUIAction(
+			FExecuteAction::CreateSP(this, &FCompositeLayerPlateCustomization::CreateCompositeMeshActor),
+			FCanExecuteAction::CreateSP(this, &FCompositeLayerPlateCustomization::CanCreateCompositeMeshActor))
+	);
+}
+
+void FCompositeLayerPlateCustomization::CreateCompositeMeshActor()
+{
+	TSharedPtr<IDetailLayoutBuilder> PinnedDetailBuilder = CachedDetailBuilder.Pin();
+	if (!PinnedDetailBuilder.IsValid())
+	{
+		return;
+	}
+	
+	TArray<TWeakObjectPtr<UCompositeLayerPlate>> Objects = PinnedDetailBuilder->GetObjectsOfTypeBeingCustomized<UCompositeLayerPlate>();
+	if (Objects.Num() != 1)
+	{
+		return;
+	}
+
+	TStrongObjectPtr<UCompositeLayerPlate> Plate = Objects[0].Pin();
+	if (!Plate.IsValid())
+	{
+		return;
+	}
+
+	ACompositeActor* CompositeActor = Plate->GetTypedOuter<ACompositeActor>();
+	if (!CompositeActor)
+	{
+		return;
+	}
+
+	FTransform SpawnTransform = FTransform::Identity;
+	if (UCameraComponent* CameraComponent = Cast<UCameraComponent>(CompositeActor->GetCamera().GetComponent(nullptr)))
+	{
+		// Size of one wall of the composite mesh
+		constexpr float MeshSize = 1000.0f;
+		
+		const float FieldOfView = FMath::DegreesToRadians(CameraComponent->FieldOfView);
+
+		// The distance from the camera a wall of MeshSize needs to be to fill the camera's horizontal field of view
+		const float Distance = 0.5f * MeshSize / FMath::Tan(0.5f * FieldOfView);
+		
+		const FVector Forward = CameraComponent->GetForwardVector();
+		const FVector Right = CameraComponent->GetRightVector();
+		const FVector Up = CameraComponent->GetUpVector();
+		
+		const FVector CameraLoc = CameraComponent->GetComponentLocation();
+
+		// Mesh location is set so that the center of the mesh's "front" wall aligns with the camera, and is
+		// a distance from it such that the most of the full wall is visible from the camera
+		const FVector MeshLoc = CameraLoc - 0.5 * MeshSize * Up - 0.5 * MeshSize * Right + (Distance - MeshSize) * Forward;
+		
+		SpawnTransform.SetLocation(MeshLoc);
+		SpawnTransform.SetRotation(CameraComponent->GetComponentRotation().Quaternion());
+	}
+
+	FScopedTransaction AddCompositeMeshTransaction(LOCTEXT("AddCompositeMeshTransaction", "Add Composite Mesh"));
+	ACompositeMeshActor* AddedMeshActor = Cast<ACompositeMeshActor>(CompositeActor->GetWorld()->SpawnActor(ACompositeMeshActor::StaticClass(), &SpawnTransform));
+	if (AddedMeshActor)
+	{
+		Plate->Modify();
+
+		FProperty* CompositeMeshesProperty = FindFProperty<FProperty>(UCompositeLayerPlate::StaticClass(), GET_MEMBER_NAME_CHECKED(UCompositeLayerPlate, CompositeMeshes));
+		Plate->PreEditChange(CompositeMeshesProperty);
+
+		Plate->CompositeMeshes.Add(AddedMeshActor);
+
+		FPropertyChangedEvent ChangedEvent(CompositeMeshesProperty, EPropertyChangeType::ArrayAdd);
+		Plate->PostEditChangeProperty(ChangedEvent);
+
+		if (GEditor)
+		{
+			GEditor->SelectNone(true, true);
+			GEditor->SelectActor(AddedMeshActor, true, true);
+			GUnrealEd->Exec(CompositeActor->GetWorld(), TEXT("CAMERA ALIGN ACTIVEVIEWPORTONLY"));
+		}
+	}
+}
+
+bool FCompositeLayerPlateCustomization::CanCreateCompositeMeshActor() const
+{
+	TSharedPtr<IDetailLayoutBuilder> PinnedDetailBuilder = CachedDetailBuilder.Pin();
+	if (!PinnedDetailBuilder.IsValid())
+	{
+		return false;
+	}
+	
+	TArray<TWeakObjectPtr<UCompositeLayerPlate>> Objects = PinnedDetailBuilder->GetObjectsOfTypeBeingCustomized<UCompositeLayerPlate>();
+	if (Objects.Num() != 1)
+	{
+		return false;
+	}
+
+	ACompositeActor* CompositeActor = Objects[0]->GetTypedOuter<ACompositeActor>();
+	if (!CompositeActor)
+	{
+		return false;
+	}
+	
+	return true;
+}
+
+void FCompositeLayerPlateCustomization::OnLayoutSizeChanged()
+{
+	TSharedPtr<IDetailLayoutBuilder> PinnedDetailBuilder = CachedDetailBuilder.Pin();
+	if (PinnedDetailBuilder.IsValid())
+	{
+		PinnedDetailBuilder->GetPropertyUtilities()->RequestRefresh();
+	}
+}
+
 #undef LOCTEXT_NAMESPACE
diff --git a/Engine/Plugins/Experimental/Compositing/Composite/Source/CompositeEditor/Private/Customizations/CompositeLayerPlateCustomization.h b/Engine/Plugins/Experimental/Compositing/Composite/Source/CompositeEditor/Private/Customizations/CompositeLayerPlateCustomization.h
--- a/Engine/Plugins/Experimental/Compositing/Composite/Source/CompositeEditor/Private/Customizations/CompositeLayerPlateCustomization.h
+++ b/Engine/Plugins/Experimental/Compositing/Composite/Source/CompositeEditor/Private/Customizations/CompositeLayerPlateCustomization.h
@@ -5,6 +5,7 @@
 #include "IDetailCustomization.h"
 
 class SCompositePlatePassPanel;
+class FMenuBuilder;
 
 /**
  * Customization for the composite plate layer, primary for displaying a custom widget for passes and composite meshes
@@ -20,6 +21,19 @@ public:
 	virtual void CustomizeDetails(IDetailLayoutBuilder& DetailLayout) override;
 	// End of IDetailCustomization interface
 
+private:
+	/** Raised by the composite mesh actor list when building the Add Actor menu */
+	void OnExtendCompositeMeshAddMenu(FMenuBuilder& MenuBuilder);
+
+	/** Creates a composite mesh actor at the appropriate position in the level and adds it to the plate's composite meshes list */
+	void CreateCompositeMeshActor();
+	
+	/** Gets whether a composite mesh actor can be created for this plate layer */
+	bool CanCreateCompositeMeshActor() const;
+
+	/** Raised when the custom widgets' layout size may have changed */
+	void OnLayoutSizeChanged();
+	
 private:
 	TWeakPtr<IDetailLayoutBuilder> CachedDetailBuilder;
 };
diff --git a/Engine/Plugins/Experimental/Compositing/Composite/Source/CompositeEditor/Private/Customizations/CompositeLayerShadowReflectionCustomization.cpp b/Engine/Plugins/Experimental/Compositing/Composite/Source/CompositeEditor/Private/Customizations/CompositeLayerShadowReflectionCustomization.cpp
new file mode 100644
--- /dev/null
+++ b/Engine/Plugins/Experimental/Compositing/Composite/Source/CompositeEditor/Private/Customizations/CompositeLayerShadowReflectionCustomization.cpp
@@ -0,0 +1,88 @@
+﻿// Copyright Epic Games, Inc. All Rights Reserved.
+
+#include "CompositeLayerShadowReflectionCustomization.h"
+
+#include "DetailCategoryBuilder.h"
+#include "DetailLayoutBuilder.h"
+#include "DetailWidgetRow.h"
+#include "IDetailGroup.h"
+#include "IPropertyUtilities.h"
+#include "Layers/CompositeLayerShadowReflection.h"
+#include "UI/SCompositeActorPickerTable.h"
+#include "Widgets/Text/STextBlock.h"
+
+#define LOCTEXT_NAMESPACE "FCompositeLayerShadowReflectionCustomization"
+
+TSharedRef<IDetailCustomization> FCompositeLayerShadowReflectionCustomization::MakeInstance()
+{
+	return MakeShared<FCompositeLayerShadowReflectionCustomization>();
+}
+
+void FCompositeLayerShadowReflectionCustomization::CustomizeDetails(const TSharedPtr<IDetailLayoutBuilder>& DetailBuilder)
+{
+	CachedDetailBuilder = DetailBuilder;
+	IDetailCustomization::CustomizeDetails(DetailBuilder);
+}
+
+void FCompositeLayerShadowReflectionCustomization::CustomizeDetails(IDetailLayoutBuilder& DetailLayout)
+{
+	DetailLayout.HideProperty(GET_MEMBER_NAME_CHECKED(UCompositeLayerShadowReflection, Actors));
+	
+	TArray<TWeakObjectPtr<UCompositeLayerShadowReflection>> Objects = DetailLayout.GetObjectsOfTypeBeingCustomized<UCompositeLayerShadowReflection>();
+
+	IDetailCategoryBuilder& LayerCategory = DetailLayout.EditCategoryAllowNone(NAME_None);
+	
+	// Add all default simple properties ahead of the custom group
+	TArray<TSharedRef<IPropertyHandle>> PlateProperties;
+	LayerCategory.GetDefaultProperties(PlateProperties, /* bSimpleProperties */ true, /* bAdvancedProperties */ false);
+	for (const TSharedRef<IPropertyHandle>& Property : PlateProperties)
+	{
+		if (!Property->IsCustomized())
+		{
+			LayerCategory.AddProperty(Property);
+		}
+	}
+
+	if (Objects.Num() == 1 && Objects[0].IsValid())
+	{
+		UCompositeLayerShadowReflection* ShadowReflection = Cast<UCompositeLayerShadowReflection>(Objects[0].Get());
+		
+		IDetailGroup& ActorListGroup = LayerCategory.AddGroup("ShadowReflectionCatcherContent", LOCTEXT("ActorListGroupName", "Shadow/Reflection Catcher Content"), false, true);
+		FCompositeActorPickerListRef ActorListRef(ShadowReflection, GET_MEMBER_NAME_CHECKED(UCompositeLayerShadowReflection, Actors), &ShadowReflection->Actors);
+		
+		ActorListGroup.AddWidgetRow()
+		.WholeRowContent()
+		[
+			SNew(SCompositeActorPickerTable, ActorListRef)
+			.OnLayoutSizeChanged(this, &FCompositeLayerShadowReflectionCustomization::OnLayoutSizeChanged)
+		];
+	}
+	else
+	{
+		// Can't display actor list if multiple layers are selected, so simply put a "Multiple Values" entry in the property list
+		LayerCategory.AddCustomRow(LOCTEXT("ActorListGroupName", "Shadow/Reflection Catcher Content"))
+		.NameContent()
+		[
+			SNew(STextBlock)
+			.Text(LOCTEXT("ActorListGroupName", "Shadow/Reflection Catcher Content"))
+			.Font(IDetailLayoutBuilder::GetDetailFont())
+		]
+		.ValueContent()
+		[
+			SNew(STextBlock)
+			.Text(LOCTEXT("MultipleValues", "Multiple Values"))
+			.Font(IDetailLayoutBuilder::GetDetailFont())
+		];
+	}
+}
+
+void FCompositeLayerShadowReflectionCustomization::OnLayoutSizeChanged()
+{
+	TSharedPtr<IDetailLayoutBuilder> PinnedDetailBuilder = CachedDetailBuilder.Pin();
+	if (PinnedDetailBuilder.IsValid())
+	{
+		PinnedDetailBuilder->GetPropertyUtilities()->RequestRefresh();
+	}
+}
+
+#undef LOCTEXT_NAMESPACE
diff --git a/Engine/Plugins/Experimental/Compositing/Composite/Source/CompositeEditor/Private/Customizations/CompositeLayerShadowReflectionCustomization.h b/Engine/Plugins/Experimental/Compositing/Composite/Source/CompositeEditor/Private/Customizations/CompositeLayerShadowReflectionCustomization.h
new file mode 100644
--- /dev/null
+++ b/Engine/Plugins/Experimental/Compositing/Composite/Source/CompositeEditor/Private/Customizations/CompositeLayerShadowReflectionCustomization.h
@@ -0,0 +1,27 @@
+﻿// Copyright Epic Games, Inc. All Rights Reserved.
+
+#pragma once
+
+#include "IDetailCustomization.h"
+
+/**
+ * Customization for the composite shadow/reflection layer, primary for displaying a custom widget for picking actors
+ */
+class FCompositeLayerShadowReflectionCustomization : public IDetailCustomization
+{
+public:
+	/** Makes a new instance of the details customization */
+	static TSharedRef<IDetailCustomization> MakeInstance();
+
+	// IDetailCustomization interface
+	virtual void CustomizeDetails(const TSharedPtr<IDetailLayoutBuilder>& DetailBuilder) override;
+	virtual void CustomizeDetails(IDetailLayoutBuilder& DetailLayout) override;
+	// End of IDetailCustomization interface
+
+private:
+	/** Raised when the custom widgets' layout size may have changed */
+	void OnLayoutSizeChanged();
+	
+private:
+	TWeakPtr<IDetailLayoutBuilder> CachedDetailBuilder;
+};
diff --git a/Engine/Plugins/Experimental/Compositing/Composite/Source/CompositeEditor/Private/UI/SCompositeActorPickerSceneOutliner.cpp b/Engine/Plugins/Experimental/Compositing/Composite/Source/CompositeEditor/Private/UI/SCompositeActorPickerSceneOutliner.cpp
new file mode 100644
--- /dev/null
+++ b/Engine/Plugins/Experimental/Compositing/Composite/Source/CompositeEditor/Private/UI/SCompositeActorPickerSceneOutliner.cpp
@@ -0,0 +1,206 @@
+﻿// Copyright Epic Games, Inc. All Rights Reserved.
+
+#include "SCompositeActorPickerSceneOutliner.h"
+
+#include "ActorMode.h"
+#include "ActorTreeItem.h"
+#include "EngineUtils.h"
+#include "ISceneOutlinerColumn.h"
+#include "ISceneOutlinerMode.h"
+#include "SceneOutlinerModule.h"
+#include "ScopedTransaction.h"
+#include "WorldTreeItem.h"
+#include "Framework/Application/SlateApplication.h"
+#include "Modules/ModuleManager.h"
+#include "Widgets/Input/SCheckBox.h"
+#include "Widgets/Layout/SBox.h"
+
+#define LOCTEXT_NAMESPACE "SCompositeActorPickerTable"
+
+/** Custom column for the scene outliner that displays a checkbox to allow users to select multiple actors quickly */
+class FCompositeActorPickerSelectedColumn : public ISceneOutlinerColumn
+{
+public:
+	DECLARE_DELEGATE_RetVal_OneParam(bool, FOnGetActorSelected, AActor*);
+	DECLARE_DELEGATE_TwoParams(FOnActorSelected, AActor*, bool /* bIsSelected */);
+	
+public:
+	FCompositeActorPickerSelectedColumn(ISceneOutliner& InSceneOutliner)
+	{
+		SceneOutliner = StaticCastSharedRef<ISceneOutliner>(InSceneOutliner.AsShared());
+	}
+
+	static FName GetID() { return "IsSelected"; }
+
+	// ISceneOutlinerColumn interface
+	virtual FName GetColumnID() override
+	{
+		return GetID();
+	}
+	
+	virtual SHeaderRow::FColumn::FArguments ConstructHeaderRowColumn() override
+	{
+		return SHeaderRow::Column(GetColumnID())
+			.FixedWidth(24.0f)
+			.DefaultLabel(FText::GetEmpty())
+			.HAlignHeader(HAlign_Left)
+			.VAlignHeader(VAlign_Center)
+			.HAlignCell(HAlign_Center)
+			.VAlignCell(VAlign_Center)
+			[
+				SNew(SCheckBox)
+				.IsChecked_Lambda([] { return ECheckBoxState::Checked; })
+			];
+	}
+
+	virtual const TSharedRef<SWidget> ConstructRowWidget(FSceneOutlinerTreeItemRef InTreeItem, const STableRow<FSceneOutlinerTreeItemPtr>& InRow) override
+	{
+		if (InTreeItem->IsA<FActorTreeItem>())
+		{
+			return SNew(SCheckBox)
+				.IsChecked(this, &FCompositeActorPickerSelectedColumn::GetTreeItemSelectedCheckState, InTreeItem)
+				.OnCheckStateChanged(this, &FCompositeActorPickerSelectedColumn::OnTreeItemSelectedCheckStateChanged, InTreeItem);
+		}
+
+		return SNullWidget::NullWidget;
+	}
+	// End of ISceneOutlinerColumn interface
+	
+private:
+	ECheckBoxState GetTreeItemSelectedCheckState(FSceneOutlinerTreeItemRef InTreeItem) const
+	{
+		if (FActorTreeItem* ActorTreeItem = InTreeItem->CastTo<FActorTreeItem>())
+		{
+			if (!OnGetActorSelected.IsBound())
+			{
+				return ECheckBoxState::Unchecked;
+			}
+			
+			if (TStrongObjectPtr<AActor> PinnedActor = ActorTreeItem->Actor.Pin())
+			{
+				return OnGetActorSelected.Execute(PinnedActor.Get()) ? ECheckBoxState::Checked : ECheckBoxState::Unchecked;
+			}
+
+			return ECheckBoxState::Unchecked;
+		}
+		
+		return ECheckBoxState::Unchecked;
+	}
+
+	void OnTreeItemSelectedCheckStateChanged(ECheckBoxState InNewCheckState, FSceneOutlinerTreeItemRef InTreeItem)
+	{
+		if (FActorTreeItem* ActorTreeItem = InTreeItem->CastTo<FActorTreeItem>())
+		{
+			if (TStrongObjectPtr<AActor> PinnedActor = ActorTreeItem->Actor.Pin())
+			{
+				OnActorSelected.ExecuteIfBound(PinnedActor.Get(), InNewCheckState == ECheckBoxState::Checked);
+			}
+		}
+	}
+	
+private:
+	TWeakPtr<ISceneOutliner> SceneOutliner = nullptr;
+
+public:
+	FOnGetActorSelected OnGetActorSelected;
+	FOnActorSelected OnActorSelected;
+};
+
+void SCompositeActorPickerSceneOutliner::Construct(const FArguments& InArgs, const FCompositeActorPickerListRef& InActorListRef)
+{
+	ActorListRef = InActorListRef;
+	OnActorListChanged = InArgs._OnActorListChanged;
+	
+	FSceneOutlinerModule& SceneOutlinerModule = FModuleManager::LoadModuleChecked<FSceneOutlinerModule>("SceneOutliner");
+
+	FSceneOutlinerInitializationOptions InitOptions;
+	{
+		InitOptions.bShowHeaderRow = true;
+		InitOptions.bShowSearchBox = true;
+		InitOptions.bShowCreateNewFolder = false;
+		InitOptions.bFocusSearchBoxWhenOpened = true;
+
+		InitOptions.ColumnMap.Add(FCompositeActorPickerSelectedColumn::GetID(), FSceneOutlinerColumnInfo(
+			ESceneOutlinerColumnVisibility::Visible,
+			0,
+			FCreateSceneOutlinerColumn::CreateLambda([this](ISceneOutliner& Outliner)
+			{
+				TSharedPtr<FCompositeActorPickerSelectedColumn> Column = MakeShared<FCompositeActorPickerSelectedColumn>(Outliner);
+				Column->OnGetActorSelected.BindSP(this, &SCompositeActorPickerSceneOutliner::OnGetActorSelected);
+				Column->OnActorSelected.BindSP(this, &SCompositeActorPickerSceneOutliner::OnActorSelected);
+				
+				return Column.ToSharedRef();
+			})
+		));
+			
+		InitOptions.ColumnMap.Add(FSceneOutlinerBuiltInColumnTypes::Label(), FSceneOutlinerColumnInfo(ESceneOutlinerColumnVisibility::Visible, 1));
+		InitOptions.ColumnMap.Add(FSceneOutlinerBuiltInColumnTypes::ActorInfo(), FSceneOutlinerColumnInfo(ESceneOutlinerColumnVisibility::Visible, 2));
+	}
+	
+	ChildSlot
+	[
+		SNew(SBox)
+		.MaxDesiredHeight(400.0f)
+		.WidthOverride(300.0f)
+		[
+			SceneOutlinerModule.CreateActorPicker(
+				InitOptions,
+				FOnActorPicked::CreateLambda([&](AActor* Actor)
+				{
+					OnActorSelected(Actor, true);
+					FSlateApplication::Get().DismissAllMenus();
+				})
+			)
+		]
+	];
+}
+
+bool SCompositeActorPickerSceneOutliner::OnGetActorSelected(AActor* InActor)
+{
+	if (!ActorListRef.IsValid())
+	{
+		return false;
+	}
+
+	return ActorListRef.ActorList->Contains(InActor);
+}
+
+void SCompositeActorPickerSceneOutliner::OnActorSelected(AActor* InActor, bool bIsSelected)
+{
+	if (!ActorListRef.IsValid())
+	{
+		return;
+	}
+	
+	TStrongObjectPtr<UObject> PinnedListOwner = ActorListRef.ActorListOwner.Pin();
+	
+	const bool bIsAdding = bIsSelected && !ActorListRef.ActorList->Contains(InActor);
+	const bool bIsRemoving = !bIsSelected && ActorListRef.ActorList->Contains(InActor);
+
+	if (!bIsAdding && !bIsRemoving)
+	{
+		return;
+	}
+	
+	FScopedTransaction SelectActorTransaction(LOCTEXT("SelectActorTransaction", "Select Actor"));
+	PinnedListOwner->Modify();
+
+	ActorListRef.NotifyPreEditChange();
+
+	EPropertyChangeType::Type ChangeType = EPropertyChangeType::Unspecified;
+	if (bIsAdding)
+	{
+		ActorListRef.ActorList->Add(InActor);
+		ChangeType = EPropertyChangeType::ArrayAdd;
+	}
+	else if (bIsRemoving)
+	{
+		ActorListRef.ActorList->Remove(InActor);
+		ChangeType = EPropertyChangeType::ArrayRemove;
+	}
+
+	ActorListRef.NotifyPostEditChangeList(ChangeType);
+	OnActorListChanged.ExecuteIfBound();
+}
+
+#undef LOCTEXT_NAMESPACE
\ No newline at end of file
diff --git a/Engine/Plugins/Experimental/Compositing/Composite/Source/CompositeEditor/Private/UI/SCompositeActorPickerSceneOutliner.h b/Engine/Plugins/Experimental/Compositing/Composite/Source/CompositeEditor/Private/UI/SCompositeActorPickerSceneOutliner.h
new file mode 100644
--- /dev/null
+++ b/Engine/Plugins/Experimental/Compositing/Composite/Source/CompositeEditor/Private/UI/SCompositeActorPickerSceneOutliner.h
@@ -0,0 +1,33 @@
+﻿// Copyright Epic Games, Inc. All Rights Reserved.
+
+#pragma once
+
+#include "SCompositeActorPickerTable.h"
+#include "Widgets/SCompoundWidget.h"
+
+/**
+ * Scene Outliner used for picking multiple actors quickly for a composite actor list
+ */
+class SCompositeActorPickerSceneOutliner : public SCompoundWidget
+{
+public:
+	SLATE_BEGIN_ARGS(SCompositeActorPickerSceneOutliner) { }
+		SLATE_EVENT(FSimpleDelegate, OnActorListChanged)
+	SLATE_END_ARGS()
+
+	void Construct(const FArguments& InArgs, const FCompositeActorPickerListRef& InActorListRef);
+
+private:
+	/** Raised when the selection state of a specific actor needs to be determined */
+	bool OnGetActorSelected(AActor* InActor);
+
+	/** Raised when the selection state of a specific actor has changed */
+	void OnActorSelected(AActor* InActor, bool bIsSelected);
+	
+private:
+	/** Reference to the actor list being modified by this actor picker */
+	FCompositeActorPickerListRef ActorListRef;
+
+	/** Raised when the actor list has changed */
+	FSimpleDelegate OnActorListChanged;
+};
diff --git a/Engine/Plugins/Experimental/Compositing/Composite/Source/CompositeEditor/Private/UI/SCompositeActorPickerTable.cpp b/Engine/Plugins/Experimental/Compositing/Composite/Source/CompositeEditor/Private/UI/SCompositeActorPickerTable.cpp
new file mode 100644
--- /dev/null
+++ b/Engine/Plugins/Experimental/Compositing/Composite/Source/CompositeEditor/Private/UI/SCompositeActorPickerTable.cpp
@@ -0,0 +1,762 @@
+﻿// Copyright Epic Games, Inc. All Rights Reserved.
+
+#include "SCompositeActorPickerTable.h"
+
+#include "CompositeEditorCommands.h"
+#include "Editor.h"
+#include "EditorActorFolders.h"
+#include "SceneOutlinerDragDrop.h"
+#include "SCompositeActorPickerSceneOutliner.h"
+#include "ScopedTransaction.h"
+#include "Selection.h"
+#include "Components/VerticalBox.h"
+#include "DragAndDrop/ActorDragDropGraphEdOp.h"
+#include "DragAndDrop/FolderDragDropOp.h"
+#include "Editor/EditorWidgets/Public/SDropTarget.h"
+#include "Filters/SFilterSearchBox.h"
+#include "Framework/MultiBox/MultiBoxBuilder.h"
+#include "Misc/TextFilter.h"
+#include "Modules/ModuleManager.h"
+#include "Styling/SlateIconFinder.h"
+#include "Widgets/Images/SLayeredImage.h"
+#include "Widgets/Input/SCheckBox.h"
+#include "Widgets/Input/SComboButton.h"
+#include "Widgets/Views/SListView.h"
+
+#define LOCTEXT_NAMESPACE "SCompositeActorPickerTable"
+
+namespace CompositeActorPickerTable
+{
+	static FName ActorListColumn_ActorName = "ActorName";
+	static FName ActorListColumn_Type = "Type";
+}
+
+/** Toolbar widget that contains filtering and an add button */
+class SCompositeActorPickerToolbar : public SCompoundWidget
+{
+private:
+	using FFilterType = SCompositeActorPickerTable::FActorListItemRef;
+	using FTreeFilter =  TTextFilter<FFilterType>;
+	using FOnExtendAddMenu = SCompositeActorPickerTable::FOnExtendAddMenu;
+	
+public:	
+	SLATE_BEGIN_ARGS(SCompositeActorPickerToolbar) {}
+		SLATE_EVENT(FSimpleDelegate, OnFilterChanged)
+		SLATE_EVENT(FSimpleDelegate, OnActorListChanged)
+		SLATE_EVENT(FOnExtendAddMenu, OnExtendAddMenu)
+	SLATE_END_ARGS()
+
+	void Construct(const FArguments& InArgs, const FCompositeActorPickerListRef& InActorListRef)
+	{
+		ActorListRef = InActorListRef;
+		OnFilterChanged = InArgs._OnFilterChanged;
+		OnActorListChanged = InArgs._OnActorListChanged;
+		OnExtendAddMenu = InArgs._OnExtendAddMenu;
+
+		InitializeFilters();
+
+		// Generate a (multi-layered) icon for the "Add" menu
+		const TSharedRef<SLayeredImage> AddIcon =
+			SNew(SLayeredImage)
+			.ColorAndOpacity(FSlateColor::UseForeground())
+			.Image(FAppStyle::GetBrush("LevelEditor.OpenAddContent.Background"));
+		AddIcon->AddLayer(FAppStyle::GetBrush("LevelEditor.OpenAddContent.Overlay"));
+		
+		ChildSlot
+		[
+			SNew(SVerticalBox)
+
+			+SVerticalBox::Slot()
+			.AutoHeight()
+			[
+				SNew(SHorizontalBox)
+
+				+SHorizontalBox::Slot()
+				.AutoWidth()
+				.Padding(0.0f, 0.0f, 4.0f, 0.0f)
+				[
+					SNew(SComboButton)
+					.HasDownArrow(true)
+					.ComboButtonStyle(FAppStyle::Get(), "SimpleComboButton")
+					.ToolTipText(LOCTEXT("AddActorButtonToolTip", "Add a level actor to the list of selected actors"))
+					.VAlign(VAlign_Center)
+					.OnGetMenuContent(this, &SCompositeActorPickerToolbar::GetAddActorMenuContent)
+					.ButtonContent()
+					[
+						AddIcon
+					]
+				]
+				
+				+SHorizontalBox::Slot()
+				.FillWidth(1.0f)
+				[
+					SAssignNew(TextFilterSearchBox, SFilterSearchBox)
+					.HintText(LOCTEXT("FilterSearch", "Search..."))
+					.ToolTipText(LOCTEXT("FilterSearchHint", "Type here to search for specific actors"))
+					.OnTextChanged_Lambda([this](const FText& InText)
+					{
+						TreeTextFilter->SetRawFilterText(InText);
+					})
+				]
+			]
+		];
+	}
+
+	bool ItemPassesFilters(const FFilterType& InItem) const
+	{
+		return TreeTextFilter->PassesFilter(InItem);
+	}
+
+	void ClearFilters()
+	{
+		TextFilterSearchBox->SetText(FText::GetEmpty());
+	}
+	
+private:
+	/** Creates all the filter bar filters and initializes the text filter */
+	void InitializeFilters()
+	{
+		TreeTextFilter = MakeShared<FTreeFilter>(FTreeFilter::FItemToStringArray::CreateLambda([](const FFilterType& InItem, TArray<FString>& OutStrings)
+		{
+			if (InItem.IsValid() && InItem->Actor.IsValid())
+			{
+				OutStrings.Add(InItem->Actor->GetActorNameOrLabel());
+				OutStrings.Add(InItem->Actor->GetClass()->GetDisplayNameText().ToString());
+			}
+		}));
+		TreeTextFilter->OnChanged().AddSP(this, &SCompositeActorPickerToolbar::FilterChanged);
+	}
+
+	/** Creates the menu widget to display when the Add button is pressed */
+	TSharedRef<SWidget> GetAddActorMenuContent()
+	{
+		constexpr bool bCloseMenuAfterSelection = true;
+		FMenuBuilder MenuBuilder(bCloseMenuAfterSelection, nullptr);
+
+		MenuBuilder.BeginSection("Add", LOCTEXT("AddSectionLabel", "Add Actor"));
+		{
+			MenuBuilder.AddMenuEntry(LOCTEXT("AddSelectedEntryLabel", "Add Selected in Outliner"),
+				LOCTEXT("AddSelectionEntryToolTip", "Adds all actors currently selected in the level editor"),
+				FSlateIcon(FAppStyle::GetAppStyleSetName(),"FoliageEditMode.SetSelect"),
+				FUIAction(
+					FExecuteAction::CreateSP(this, &SCompositeActorPickerToolbar::AddSelectedActors),
+					FCanExecuteAction::CreateSP(this, &SCompositeActorPickerToolbar::CanAddSelectedActors))
+			);
+
+			OnExtendAddMenu.ExecuteIfBound(MenuBuilder);
+		}
+		MenuBuilder.EndSection();
+		
+		MenuBuilder.BeginSection("SceneOutliner", LOCTEXT("SceneOutlinerSectionLabel", "Browse"));
+		{
+			TSharedRef<SWidget> SceneOutliner =
+				SNew(SCompositeActorPickerSceneOutliner, ActorListRef)
+				.OnActorListChanged(OnActorListChanged);
+			
+			MenuBuilder.AddWidget(SceneOutliner, FText::GetEmpty(), true);
+		}
+		MenuBuilder.EndSection();
+		
+		return MenuBuilder.MakeWidget();
+	}
+
+	/** Adds all actors selected in the level/outliner to the actor list */
+	void AddSelectedActors()
+	{
+		if (!ActorListRef.IsValid())
+		{
+			return;
+		}
+
+		TStrongObjectPtr<UObject> PinnedListOwner = ActorListRef.ActorListOwner.Pin();
+		
+		TArray<AActor*> SelectedActors;
+		GEditor->GetSelectedActors()->GetSelectedObjects<AActor>(SelectedActors);
+
+		if (SelectedActors.IsEmpty())
+		{
+			return;
+		}
+
+		TSharedPtr<FScopedTransaction> AddActorsTransaction = nullptr;
+
+		for (AActor* SelectedActor : SelectedActors)
+		{
+			if (ActorListRef.ActorList->Contains(SelectedActor))
+			{
+				continue;
+			}
+
+			if (!AddActorsTransaction.IsValid())
+			{
+				AddActorsTransaction = MakeShared<FScopedTransaction>(LOCTEXT("AddActorsTransaction", "Add Actors"));
+				PinnedListOwner->Modify();
+
+				ActorListRef.NotifyPreEditChange();
+			}
+
+			ActorListRef.ActorList->Add(SelectedActor);
+		}
+
+		if (AddActorsTransaction.IsValid())
+		{
+			ActorListRef.NotifyPostEditChangeList(EPropertyChangeType::ArrayAdd);
+			OnActorListChanged.ExecuteIfBound();
+		}
+	}
+
+	/** Gets whether the selected level/outliner actors can be added to the actor list */
+	bool CanAddSelectedActors()
+	{
+		TArray<AActor*> SelectedActors;
+		GEditor->GetSelectedActors()->GetSelectedObjects<AActor>(SelectedActors);
+		
+		return ActorListRef.IsValid() && SelectedActors.Num() > 0;
+	}
+
+	/** Raised when the actor list filter has been changed */
+	void FilterChanged()
+	{
+		OnFilterChanged.ExecuteIfBound();
+	}
+	
+private:
+	/** A reference to the actor list being managed by the actor picker */
+	FCompositeActorPickerListRef ActorListRef;
+	
+	TSharedPtr<FTreeFilter> TreeTextFilter;
+	
+	TSharedPtr<SFilterSearchBox> TextFilterSearchBox;
+
+	FSimpleDelegate OnFilterChanged;
+	FSimpleDelegate OnActorListChanged;
+	FOnExtendAddMenu OnExtendAddMenu;
+};
+
+/** Table row to display an actor in the actor list view */
+class SCompositeActorListItemRow : public SMultiColumnTableRow<SCompositeActorPickerTable::FActorListItemRef>
+{
+	using FActorListItemRef = SCompositeActorPickerTable::FActorListItemRef;
+	
+public:
+	SLATE_BEGIN_ARGS(SCompositeActorListItemRow) { }
+	SLATE_END_ARGS()
+
+	void Construct(const FArguments& InArgs, const TSharedRef<STableViewBase>& InOwnerTable, const FActorListItemRef& InActorItem)
+	{
+		ActorItem = InActorItem;
+
+		STableRow<FActorListItemRef>::FArguments Args = FSuperRowType::FArguments()
+		   .Style(&FAppStyle::Get().GetWidgetStyle<FTableRowStyle>("SceneOutliner.TableViewRow"));
+		
+		SMultiColumnTableRow::Construct(Args, InOwnerTable);
+	}
+
+	virtual TSharedRef<SWidget> GenerateWidgetForColumn(const FName& InColumnName) override
+	{
+		if (InColumnName == CompositeActorPickerTable::ActorListColumn_ActorName)
+		{
+			return SNew(SHorizontalBox)
+			
+				+SHorizontalBox::Slot()
+				.AutoWidth()
+				.VAlign(VAlign_Center)
+				.Padding(8.0f, 1.0f, 6.0f, 1.0f)
+				[
+					SNew(SBox)
+					.WidthOverride(16.0f)
+					.HeightOverride(16.0f)
+					[
+						SNew(SImage).Image(this, &SCompositeActorListItemRow::GetActorIcon)
+					]
+				]
+
+				+SHorizontalBox::Slot()
+				.VAlign(VAlign_Center)
+				[
+					SNew(STextBlock).Text(this, &SCompositeActorListItemRow::GetActorName)
+				];
+		}
+		else if (InColumnName == CompositeActorPickerTable::ActorListColumn_Type)
+		{
+			return SNew(STextBlock).Text(this, &SCompositeActorListItemRow::GetActorType);
+		}
+
+		return SNullWidget::NullWidget;
+	}
+
+private:
+	/** Gets the icon to display for the actor */
+	const FSlateBrush* GetActorIcon() const
+	{
+		if (ActorItem.IsValid() && ActorItem->Actor.IsValid())
+		{
+			return FSlateIconFinder::FindIconForClass(ActorItem->Actor->GetClass()).GetIcon();
+		}
+
+		return nullptr;
+	}
+
+	/** Gets the name to display for the actor */
+	FText GetActorName() const
+	{
+		if (ActorItem.IsValid() && ActorItem->Actor.IsValid())
+		{
+			return FText::FromString(ActorItem->Actor->GetActorNameOrLabel());
+		}
+
+		return LOCTEXT("NoneActorLabel", "None");
+	}
+
+	/** Gets the name of the actor's type to display */
+	FText GetActorType() const
+	{
+		if (ActorItem.IsValid() && ActorItem->Actor.IsValid())
+		{
+			return ActorItem->Actor->GetClass()->GetDisplayNameText();
+		}
+
+		return LOCTEXT("InvalidActorTypeLabel", "-");
+	}
+	
+private:
+	FActorListItemRef ActorItem;
+};
+
+FCompositeActorPickerListRef::FCompositeActorPickerListRef(const TWeakObjectPtr<UObject>& InActorListOwner, const FName& InActorListPropertyName, TArray<TObjectPtr<AActor>>* InActorList)
+	: ActorListOwner(InActorListOwner)
+	, ActorListPropertyName(InActorListPropertyName)
+	, ActorList(InActorList)
+{
+	if (TStrongObjectPtr<UObject> PinnedOwner = ActorListOwner.Pin())
+	{
+		ActorListProperty = FindFProperty<FProperty>(PinnedOwner->GetClass(), ActorListPropertyName);
+	}
+}
+
+bool FCompositeActorPickerListRef::IsValid() const
+{
+	return ActorListOwner.IsValid() && ActorListProperty != nullptr && ActorList != nullptr;
+}
+
+void FCompositeActorPickerListRef::NotifyPreEditChange()
+{
+	if (!IsValid())
+	{
+		return;
+	}
+
+	if (TStrongObjectPtr<UObject> PinnedOwner = ActorListOwner.Pin())
+	{
+		PinnedOwner->PreEditChange(ActorListProperty);
+	}
+}
+
+void FCompositeActorPickerListRef::NotifyPostEditChangeList(EPropertyChangeType::Type InChangeType)
+{
+	if (!IsValid())
+	{
+		return;
+	}
+
+	if (TStrongObjectPtr<UObject> PinnedOwner = ActorListOwner.Pin())
+	{
+		FPropertyChangedEvent ChangedEvent(ActorListProperty, InChangeType);
+		PinnedOwner->PostEditChangeProperty(ChangedEvent);
+	}
+}
+
+void SCompositeActorPickerTable::Construct(const FArguments& InArgs, const FCompositeActorPickerListRef& InActorListRef)
+{
+	ActorListRef = InActorListRef;
+	OnLayoutSizeChanged = InArgs._OnLayoutSizeChanged;
+	
+	FCoreUObjectDelegates::OnObjectPropertyChanged.AddSP(this, &SCompositeActorPickerTable::OnObjectPropertyChanged);
+	
+	if (GEditor)
+	{
+		GEditor->RegisterForUndo(this);
+	}
+
+	CommandList = MakeShared<FUICommandList>();
+	BindCommands();
+	
+	ChildSlot
+	[
+		SNew(SVerticalBox)
+
+		+SVerticalBox::Slot()
+		.AutoHeight()
+		.Padding(0.0f, 4.0f)
+		[
+			SAssignNew(Toolbar, SCompositeActorPickerToolbar, ActorListRef)
+			.OnFilterChanged(this, &SCompositeActorPickerTable::OnFilterChanged)
+			.OnActorListChanged_Lambda([this]
+			{
+				FillActorList();
+			})
+			.OnExtendAddMenu(InArgs._OnExtendAddMenu)
+		]
+
+		+SVerticalBox::Slot()
+		.AutoHeight()
+		[
+			SNew(SDropTarget)
+			.OnAllowDrop(this, &SCompositeActorPickerTable::OnAllowListDrop)
+			.OnDropped(this, &SCompositeActorPickerTable::OnListDropped)
+			[
+				SAssignNew(ListView, SListView<FActorListItemRef>)
+				.HeaderRow(
+					SNew(SHeaderRow)
+
+					+SHeaderRow::Column(CompositeActorPickerTable::ActorListColumn_ActorName)
+					.FillWidth(0.6)
+					.DefaultLabel(LOCTEXT("ActorNameColumnLabel", "Item Label"))
+
+					+SHeaderRow::Column(CompositeActorPickerTable::ActorListColumn_Type)
+					.FillWidth(0.4)
+					.DefaultLabel(LOCTEXT("ActorTypeColumnLabel", "Type"))
+				)
+				.ListItemsSource(&FilteredActorListItems)
+				.OnGenerateRow_Lambda([](FActorListItemRef InListItem, const TSharedRef<STableViewBase>& OwnerTable)
+				{
+					return SNew(SCompositeActorListItemRow, OwnerTable, InListItem);
+				})
+				.OnContextMenuOpening(this, &SCompositeActorPickerTable::CreateListContextMenu)
+			]
+		]
+
+		+SVerticalBox::Slot()
+		.AutoHeight()
+		[
+			SNew(STextBlock)
+			.Text(this, &SCompositeActorPickerTable::GetFilterStatusText)
+			.ColorAndOpacity(this, &SCompositeActorPickerTable::GetFilterStatusTextColor)
+		]
+	];
+
+	FillActorList();
+}
+
+SCompositeActorPickerTable::~SCompositeActorPickerTable()
+{
+	if (GEditor)
+	{
+		GEditor->UnregisterForUndo(this);
+	}
+	
+	FCoreUObjectDelegates::OnObjectPropertyChanged.RemoveAll(this);
+}
+
+void SCompositeActorPickerTable::Tick(const FGeometry& AllottedGeometry, const double InCurrentTime, const float InDeltaTime)
+{
+	SCompoundWidget::Tick(AllottedGeometry, InCurrentTime, InDeltaTime);
+
+	// Keep track of when the list view has changed the number of widgets it is actively displaying.
+	// This allows us to know if list items have been added or removed (as a list view's internal list is
+	// not generated immediately when calling RefreshList), so that we can call OnLayoutSizeChanged at the right time
+	const int32 ListViewNumChildren = ListView.IsValid() ? ListView->GetNumGeneratedChildren() : 0;
+	if (ListViewNumChildren != CachedListViewNumChildren)
+	{
+		OnLayoutSizeChanged.ExecuteIfBound();
+		CachedListViewNumChildren = ListViewNumChildren;
+	}
+}
+
+FReply SCompositeActorPickerTable::OnKeyDown(const FGeometry& MyGeometry, const FKeyEvent& InKeyEvent)
+{
+	if (CommandList->ProcessCommandBindings(InKeyEvent))
+	{
+		return FReply::Handled();
+	}
+	
+	return FReply::Unhandled();
+}
+
+void SCompositeActorPickerTable::PostUndo(bool bSuccess)
+{
+	FillActorList();
+}
+
+void SCompositeActorPickerTable::PostRedo(bool bSuccess)
+{
+	PostUndo(bSuccess);
+}
+
+void SCompositeActorPickerTable::BindCommands()
+{
+	CommandList->MapAction(FCompositeEditorCommands::Get().RemoveActor,
+		FExecuteAction::CreateSP(this, &SCompositeActorPickerTable::RemoveActors),
+		FCanExecuteAction::CreateSP(this, &SCompositeActorPickerTable::CanRemoveActors));
+}
+
+void SCompositeActorPickerTable::FillActorList()
+{
+	ActorListItems.Empty();
+	
+	if (!ActorListRef.IsValid())
+	{
+		return;
+	}
+
+	for (int32 Index = 0; Index < ActorListRef.ActorList->Num(); ++Index)
+	{
+		ActorListItems.Add(MakeShared<FActorListItem>((*ActorListRef.ActorList)[Index], Index));
+	}
+
+	FilterActorList();
+
+	if (ListView.IsValid())
+	{
+		ListView->RequestListRefresh();
+	}
+}
+
+void SCompositeActorPickerTable::FilterActorList()
+{
+	FilteredActorListItems.Empty();
+
+	for (const FActorListItemRef& Actor : ActorListItems)
+	{
+		const bool bPassesFilter = Toolbar.IsValid() ? Toolbar->ItemPassesFilters(Actor) : true;
+		if (bPassesFilter)
+		{
+			FilteredActorListItems.Add(Actor);
+		}
+	}
+}
+
+void SCompositeActorPickerTable::OnFilterChanged()
+{
+	FilterActorList();
+	
+	if (ListView.IsValid())
+	{
+		ListView->RequestListRefresh();
+	}
+}
+
+TSharedPtr<SWidget> SCompositeActorPickerTable::CreateListContextMenu()
+{
+	if (!ListView.IsValid())
+	{
+		return SNullWidget::NullWidget;
+	}
+	
+	const int32 NumItems = ListView->GetNumItemsSelected();
+	if (NumItems >= 1)
+	{
+		constexpr bool bShouldCloseWindowAfterMenuSelection = true;
+		FMenuBuilder MenuBuilder(bShouldCloseWindowAfterMenuSelection, CommandList);
+		
+		MenuBuilder.AddMenuEntry(FCompositeEditorCommands::Get().RemoveActor);
+		
+		return MenuBuilder.MakeWidget();
+	}
+
+	return SNullWidget::NullWidget;
+}
+
+bool SCompositeActorPickerTable::OnAllowListDrop(TSharedPtr<FDragDropOperation> InDragDropOperation)
+{
+	// Support dragging both actors and folders from the Outliner (dragging a folder will add all actor types in the folder)
+	return InDragDropOperation->IsOfType<FActorDragDropGraphEdOp>() || InDragDropOperation->IsOfType<FSceneOutlinerDragDropOp>();
+}
+
+FReply SCompositeActorPickerTable::OnListDropped(const FGeometry& Geometry, const FDragDropEvent& DragDropEvent)
+{
+	if (!ActorListRef.IsValid())
+	{
+		return FReply::Unhandled();
+	}
+
+	TStrongObjectPtr<UObject> PinnedListOwner = ActorListRef.ActorListOwner.Pin();
+	
+	TArray<AActor*> DroppedActors;
+	if (const TSharedPtr<FActorDragDropGraphEdOp> ActorOperation = DragDropEvent.GetOperationAs<FActorDragDropGraphEdOp>())
+	{
+		Algo::TransformIf(ActorOperation->Actors, DroppedActors,
+			[](const TWeakObjectPtr<AActor>& Actor) { return Actor.IsValid() && !Actor->HasAnyFlags(RF_Transient); },
+			[](const TWeakObjectPtr<AActor>& Actor) { return Actor.Get(); });
+	}
+	else if (const TSharedPtr<FSceneOutlinerDragDropOp> SceneOperation = DragDropEvent.GetOperationAs<FSceneOutlinerDragDropOp>())
+	{
+		if (const TSharedPtr<FFolderDragDropOp> FolderOp = SceneOperation->GetSubOp<FFolderDragDropOp>())
+		{
+			FActorFolders::GetActorsFromFolders(*FolderOp->World.Get(), FolderOp->Folders, DroppedActors);
+			DroppedActors.RemoveAll([](const AActor* Actor) { return Actor->HasAnyFlags(RF_Transient); });
+		}
+	}
+			
+	TSharedPtr<FScopedTransaction> AddActorsTransaction = nullptr;
+	for (AActor* DroppedActor : DroppedActors)
+	{
+		if (ActorListRef.ActorList->Contains(DroppedActor))
+		{
+			continue;
+		}
+
+		if (!AddActorsTransaction.IsValid())
+		{
+			AddActorsTransaction = MakeShared<FScopedTransaction>(LOCTEXT("AddActorsTransaction", "Add Actors"));
+			PinnedListOwner->Modify();
+
+			ActorListRef.NotifyPreEditChange();
+		}
+
+		ActorListRef.ActorList->Add(DroppedActor);
+	}
+
+	if (AddActorsTransaction.IsValid())
+	{
+		ActorListRef.NotifyPostEditChangeList(EPropertyChangeType::ArrayAdd);
+		FillActorList();
+	}
+
+	return FReply::Handled();
+}
+
+FText SCompositeActorPickerTable::GetFilterStatusText() const
+{
+	const int32 NumActors = ActorListItems.Num();
+	const int32 NumFiltered = FilteredActorListItems.Num();
+	const int32 NumSelected = ListView.IsValid() ? ListView->GetNumItemsSelected() : 0;
+	
+	const FText ActorLabel = NumActors > 1 ? LOCTEXT("ActorPlural", "Actors") : LOCTEXT("ActorSingular", "Actor");
+	if (NumActors > 0 && NumActors == NumFiltered)
+	{
+		if (NumSelected > 0)
+		{
+			return FText::Format(LOCTEXT("NumActorsAndSelectedTextFormat", "{0} {1}, {2} Selected"), FText::AsNumber(NumActors), ActorLabel, FText::AsNumber(NumSelected));
+		}
+		else
+		{
+			return FText::Format(LOCTEXT("NumActorsTextFormat", "{0} {1}"), FText::AsNumber(NumActors), ActorLabel);
+		}
+	}
+	else if (NumFiltered > 0)
+	{
+		if (NumSelected > 0)
+		{
+			return FText::Format(LOCTEXT("NumActorsWithFilteredAndSelectedTextFormat", "Showing {0} of {1} {2}, {3} Selected"),
+				FText::AsNumber(NumFiltered),
+				FText::AsNumber(NumActors),
+				ActorLabel,
+				FText::AsNumber(NumSelected));
+		}
+		else
+		{
+			return FText::Format(LOCTEXT("NumActorsWithFilteredTextFormat", "Showing {0} of {1} {2}"),
+				FText::AsNumber(NumFiltered),
+				FText::AsNumber(NumActors),
+				ActorLabel);
+		}
+	}
+	else
+	{
+		if (NumActors > 0)
+		{
+			return FText::Format(LOCTEXT("NoMatchingActorsTextFormat", "No matching actors ({0} {1})"), FText::AsNumber(NumActors), ActorLabel);
+		}
+		else
+		{
+			return LOCTEXT("NoActorsFoundLabel", "No actors");
+		}
+	}
+}
+
+FSlateColor SCompositeActorPickerTable::GetFilterStatusTextColor() const
+{
+	const int32 NumActors = ActorListItems.Num();
+	const int32 NumFiltered = FilteredActorListItems.Num();
+	
+	if (NumActors == 0 || NumFiltered == 0)
+	{
+		return FAppStyle::Get().GetSlateColor("Colors.AccentRed");
+	}
+	
+	return FSlateColor::UseForeground();
+}
+
+void SCompositeActorPickerTable::RemoveActors()
+{
+	TArray<FActorListItemRef> SelectedActors = ListView.IsValid() ? ListView->GetSelectedItems() : TArray<FActorListItemRef>();
+	if (SelectedActors.IsEmpty())
+	{
+		return;
+	}
+
+	if (!ActorListRef.IsValid())
+	{
+		return;
+	}
+
+	TStrongObjectPtr<UObject> PinnedListOwner = ActorListRef.ActorListOwner.Pin();
+
+	TArray<int32> ActorIndicesToRemove;
+	for (const FActorListItemRef& ActorItem : SelectedActors)
+	{
+		if (!ActorListRef.ActorList->IsValidIndex(ActorItem->Index))
+		{
+			return;
+		}
+
+		ActorIndicesToRemove.Add(ActorItem->Index);
+	}
+
+	if (ActorIndicesToRemove.IsEmpty())
+	{
+		return;
+	}
+
+	// Remove in reverse order so that indices don't get messed up as actors are removed
+	ActorIndicesToRemove.Sort();
+
+	FScopedTransaction RemoveActorsTransaction(LOCTEXT("RemoveActorsTransaction", "Remove Actors"));
+	PinnedListOwner->Modify();
+	ActorListRef.NotifyPreEditChange();
+	
+	for (int32 Index = ActorIndicesToRemove.Num() - 1; Index >= 0; --Index)
+	{
+		ActorListRef.ActorList->RemoveAt(ActorIndicesToRemove[Index]);
+	}
+
+	ActorListRef.NotifyPostEditChangeList(EPropertyChangeType::ArrayRemove);
+	FillActorList();
+}
+
+bool SCompositeActorPickerTable::CanRemoveActors() const
+{
+	const int32 NumSelected = ListView.IsValid() ? ListView->GetNumItemsSelected() : 0;
+	if (NumSelected == 0)
+	{
+		return false;
+	}
+
+	if (!ActorListRef.IsValid())
+	{
+		return false;
+	}
+
+	return true;
+}
+
+void SCompositeActorPickerTable::OnObjectPropertyChanged(UObject* InObject, FPropertyChangedEvent& InPropertyChangedEvent)
+{
+	if (!ActorListRef.IsValid())
+	{
+		return;
+	}
+
+	if (InObject != ActorListRef.ActorListOwner.Get())
+	{
+		return;
+	}
+
+	if (InPropertyChangedEvent.GetPropertyName() == ActorListRef.ActorListPropertyName)
+	{
+		FillActorList();
+	}
+}
+
+#undef LOCTEXT_NAMESPACE
diff --git a/Engine/Plugins/Experimental/Compositing/Composite/Source/CompositeEditor/Private/UI/SCompositeActorPickerTable.h b/Engine/Plugins/Experimental/Compositing/Composite/Source/CompositeEditor/Private/UI/SCompositeActorPickerTable.h
new file mode 100644
--- /dev/null
+++ b/Engine/Plugins/Experimental/Compositing/Composite/Source/CompositeEditor/Private/UI/SCompositeActorPickerTable.h
@@ -0,0 +1,155 @@
+﻿// Copyright Epic Games, Inc. All Rights Reserved.
+
+#pragma once
+
+#include "EditorUndoClient.h"
+#include "Widgets/SCompoundWidget.h"
+
+class SCompositeActorPickerToolbar;
+class FMenuBuilder;
+class FUICommandList;
+template<typename T> class SListView;
+
+/** Stores a pointer reference to a list of actors on a UObject to be managed by the actor picker */
+struct FCompositeActorPickerListRef
+{
+public:
+	FCompositeActorPickerListRef(const TWeakObjectPtr<UObject>& InActorListOwner, const FName& InActorListPropertyName, TArray<TObjectPtr<AActor>>* InActorList);
+	
+	FCompositeActorPickerListRef() { }
+	
+	/** Gets whether the pointer to the actor list is valid */
+	bool IsValid() const;
+
+	/** Raises the owner's PreEditChange event for the actor list property */
+	void NotifyPreEditChange();
+
+	/** Raises the owner's PostEditchange event for the actor list property */
+	void NotifyPostEditChangeList(EPropertyChangeType::Type InChangeType);
+	
+public:
+	/** Reference to the UObject that owns the actor list being edited by this widget */
+	TWeakObjectPtr<UObject> ActorListOwner = nullptr;
+
+	/** The name of the actor list property within ActorListOwner */
+	FName ActorListPropertyName = NAME_None;
+
+	/** The property in ActorListOwner of the actor list */
+	FProperty* ActorListProperty = nullptr;
+	
+	/** Pointer to the actor list being edited by this widget */
+	TArray<TObjectPtr<AActor>>* ActorList = nullptr;
+};
+
+/**
+ * A table that displays a list of selected actors as well as a picker to easily add or remove actors from the current level
+ */
+class SCompositeActorPickerTable : public SCompoundWidget, public FEditorUndoClient
+{
+private:
+	struct FActorListItem
+	{
+		/** Weak reference to the actor this item represents */
+		TWeakObjectPtr<AActor> Actor;
+
+		/** The index of the actor in the list */
+		int32 Index;
+
+		FActorListItem(const TWeakObjectPtr<AActor>& InActor, int32 InIndex)
+			: Actor(InActor)
+			, Index(InIndex)
+		{ }
+	};
+	using FActorListItemRef = TSharedPtr<FActorListItem>;
+
+public:
+	DECLARE_DELEGATE_OneParam(FOnExtendAddMenu, FMenuBuilder&);
+	
+public:
+	SLATE_BEGIN_ARGS(SCompositeActorPickerTable) { }
+		SLATE_EVENT(FOnExtendAddMenu, OnExtendAddMenu)
+		SLATE_EVENT(FSimpleDelegate, OnLayoutSizeChanged)
+	SLATE_END_ARGS()
+
+	void Construct(const FArguments& InArgs, const FCompositeActorPickerListRef& InActorListRef);
+
+	virtual ~SCompositeActorPickerTable() override;
+
+	// SWidget interface
+	virtual void Tick(const FGeometry& AllottedGeometry, const double InCurrentTime, const float InDeltaTime) override;
+	virtual FReply OnKeyDown(const FGeometry& MyGeometry, const FKeyEvent& InKeyEvent) override;
+	// End of SWidget interface
+	
+	// FEditorUndoClient Interface
+	virtual void PostUndo(bool bSuccess) override;
+	virtual void PostRedo(bool bSuccess) override;
+	// End of FEditorUndoClient
+	
+private:
+	/** Binds any commands to callbacks for the actor list view */
+	void BindCommands();
+	
+	/** Fills the list view's source list with the valid actors from the actor list and applies any active filters */
+	void FillActorList();
+
+	/** Filters the actor list using any active filters */
+	void FilterActorList();
+
+	/** Raised when the toolbar filter has changed */
+	void OnFilterChanged();
+
+	/** Create the right click context menu for the actor list view */
+	TSharedPtr<SWidget> CreateListContextMenu();
+
+	/** Raised when a drag drop operation is trying to be dropped on the actor list */
+	bool OnAllowListDrop(TSharedPtr<FDragDropOperation> InDragDropOperation);
+
+	/** Raised when a drag drop operation is being dropped on the actor list */
+	FReply OnListDropped(const FGeometry& Geometry, const FDragDropEvent& DragDropEvent);
+	
+	/** Gets the filter status text */
+	FText GetFilterStatusText() const;
+
+	/** Gets the color for the filter status text */
+	FSlateColor GetFilterStatusTextColor() const;
+	
+	/** Removes the selected actors from the list */
+	void RemoveActors();
+
+	/** Gets whether the RemoveActors command can be used */
+	bool CanRemoveActors() const;
+	
+	/** Raised when a property on the object has changed */
+	void OnObjectPropertyChanged(UObject* InObject, FPropertyChangedEvent& InPropertyChangedEvent);
+	
+private:
+	/** List view that displays the list of selected actors */
+	TSharedPtr<SListView<FActorListItemRef>> ListView;
+
+	/** The toolbar that contains the Add button and the filter text box */
+	TSharedPtr<SCompositeActorPickerToolbar> Toolbar;
+	
+	/** The reference to the actor list being managed by this actor picker */
+	FCompositeActorPickerListRef ActorListRef;
+
+	/** The list of actors to display in the list view */
+	TArray<FActorListItemRef> ActorListItems;
+
+	/** Filtered list of actors to display in the list view */
+	TArray<FActorListItemRef> FilteredActorListItems;
+
+	/** The command list for commands related to the actor list view */
+	TSharedPtr<FUICommandList> CommandList;
+
+	/** Stores the last observed number of widgets the list view is actually displaying, used to determine if the list view's layout has changed */
+	int32 CachedListViewNumChildren = 0;
+	
+	/** Raised when building the Add menu, allows the menu contents to be extended with additional entries */
+	FOnExtendAddMenu OnExtendAddMenu;
+
+	/** Raised when the number of items in the list view changes */
+	FSimpleDelegate OnLayoutSizeChanged;
+
+	friend class SCompositeActorPickerToolbar;
+	friend class SCompositeActorListItemRow;
+};
diff --git a/Engine/Plugins/Experimental/Compositing/Composite/Source/CompositeEditor/Private/UI/SCompositePanelLayerTree.cpp b/Engine/Plugins/Experimental/Compositing/Composite/Source/CompositeEditor/Private/UI/SCompositePanelLayerTree.cpp
--- a/Engine/Plugins/Experimental/Compositing/Composite/Source/CompositeEditor/Private/UI/SCompositePanelLayerTree.cpp
+++ b/Engine/Plugins/Experimental/Compositing/Composite/Source/CompositeEditor/Private/UI/SCompositePanelLayerTree.cpp
@@ -1027,8 +1027,21 @@ void SCompositePanelLayerTree::FillCompositeActorTreeItem(FCompositeActorTreeIte
 	}
 }
 
-void SCompositePanelLayerTree::RefreshCompositeActorLayers(ACompositeActor* InCompositeActor, bool bRefreshTree)
+void SCompositePanelLayerTree::RefreshCompositeActorLayers(ACompositeActor* InCompositeActor, bool bRefreshTree, bool bPreserveSelection)
 {
+	TArray<int32> LayersToSelect;
+	if (bPreserveSelection)
+	{
+		TArray<FCompositeActorTreeItemPtr> SelectedItems = TreeView.IsValid() ? TreeView->GetSelectedItems() : TArray<FCompositeActorTreeItemPtr>();
+		for (const FCompositeActorTreeItemPtr& SelectedItem : SelectedItems)
+		{
+			if (SelectedItem->CompositeActor.Get() == InCompositeActor)
+			{
+				LayersToSelect.Add(SelectedItem->LayerIndex);
+			}
+		}
+	}
+	
 	int32 TreeItemIndex = CompositeActorTreeItems.IndexOfByPredicate([InCompositeActor](const FCompositeActorTreeItemPtr& InTreeItem)
 	{
 		return InTreeItem.IsValid() && InTreeItem->CompositeActor == InCompositeActor;
@@ -1042,6 +1055,19 @@ void SCompositePanelLayerTree::RefreshCompositeActorLayers(ACompositeActor* InCo
 		{
 			RefreshAndExpandTreeView();
 		}
+
+		if (!LayersToSelect.IsEmpty())
+		{
+			for (int32 LayerIndex : LayersToSelect)
+			{
+				const bool bIsValidIndex = CompositeActorTreeItems[TreeItemIndex]->Children.IsValidIndex(LayerIndex);
+				FCompositeActorTreeItemPtr LayerToSelect = bIsValidIndex ? CompositeActorTreeItems[TreeItemIndex]->Children[LayerIndex] : nullptr;
+				if (LayerToSelect.IsValid() && TreeView.IsValid())
+				{
+					TreeView->SetItemSelection(LayerToSelect, true);
+				}
+			}
+		}
 	}
 }
 
@@ -1588,7 +1614,9 @@ void SCompositePanelLayerTree::OnObjectPropertyChanged(UObject* Object, FPropert
 {
 	if (Object->IsA<ACompositeActor>() && PropertyChangedEvent.GetPropertyName() == GET_MEMBER_NAME_CHECKED(ACompositeActor, CompositeLayers))
 	{
-		RefreshCompositeActorLayers(Cast<ACompositeActor>(Object));
+		constexpr bool bRefreshTree = true;
+		constexpr bool bPreserveSelection = true;
+		RefreshCompositeActorLayers(Cast<ACompositeActor>(Object), bRefreshTree, bPreserveSelection);
 	}
 }
 
@@ -1850,7 +1878,7 @@ void SCompositePanelLayerTree::DeleteSelectedItems()
 			EditChain.AddTail(CompositeLayersProperty);
 			EditChain.SetActivePropertyNode(CompositeLayersProperty);
 
-			FPropertyChangedEvent ChangeEvent(CompositeLayersProperty, EPropertyChangeType::ArrayAdd);
+			FPropertyChangedEvent ChangeEvent(CompositeLayersProperty, EPropertyChangeType::ArrayRemove);
 			FPropertyChangedChainEvent ChainEvent(EditChain, ChangeEvent);
 			CompositeActor->PostEditChangeChainProperty(ChainEvent);
 			
diff --git a/Engine/Plugins/Experimental/Compositing/Composite/Source/CompositeEditor/Private/UI/SCompositePanelLayerTree.h b/Engine/Plugins/Experimental/Compositing/Composite/Source/CompositeEditor/Private/UI/SCompositePanelLayerTree.h
--- a/Engine/Plugins/Experimental/Compositing/Composite/Source/CompositeEditor/Private/UI/SCompositePanelLayerTree.h
+++ b/Engine/Plugins/Experimental/Compositing/Composite/Source/CompositeEditor/Private/UI/SCompositePanelLayerTree.h
@@ -76,7 +76,7 @@ private:
 	void FillCompositeActorTreeItem(FCompositeActorTreeItemPtr& InOutTreeItem);
 
 	/** Refreshes the child tree elements for the specified composite actor */
-	void RefreshCompositeActorLayers(ACompositeActor* InCompositeActor, bool bRefreshTree = true);
+	void RefreshCompositeActorLayers(ACompositeActor* InCompositeActor, bool bRefreshTree = true, bool bPreserveSelection = false);
 
 	/** Selects the specified layers of a composite actor in the tree view */
 	void SelectCompositeActorLayers(ACompositeActor* InCompositeActor, const TArray<int32>& InLayersToSelect);
diff --git a/Engine/Plugins/Experimental/Compositing/Composite/Source/CompositeEditor/Private/UI/SCompositePassTree.cpp b/Engine/Plugins/Experimental/Compositing/Composite/Source/CompositeEditor/Private/UI/SCompositePassTree.cpp
--- a/Engine/Plugins/Experimental/Compositing/Composite/Source/CompositeEditor/Private/UI/SCompositePassTree.cpp
+++ b/Engine/Plugins/Experimental/Compositing/Composite/Source/CompositeEditor/Private/UI/SCompositePassTree.cpp
@@ -600,6 +600,11 @@ private:
 		bool bAnyDisabled = false;
 		for (const TObjectPtr<UCompositePassBase>& Pass : AllPassesOfType)
 		{
+			if (!IsValid(Pass))
+			{
+				continue;
+			}
+			
 			if (Pass->IsEnabled())
 			{
 				bAnyEnabled = true;
@@ -828,6 +833,8 @@ void SCompositePassTree::Construct(const FArguments& InArgs, const TWeakObjectPt
 	OnSelectionChanged = InArgs._OnSelectionChanged;
 	OnLayoutChanged = InArgs._OnLayoutChanged;
 
+	FCoreUObjectDelegates::OnObjectPropertyChanged.AddSP(this, &SCompositePassTree::OnObjectPropertyChanged);
+	
 	CommandList = MakeShared<FUICommandList>();
 	BindCommands();
 	
@@ -902,6 +909,11 @@ void SCompositePassTree::Construct(const FArguments& InArgs, const TWeakObjectPt
 	FillPassTreeItems();
 }
 
+SCompositePassTree::~SCompositePassTree()
+{
+	FCoreUObjectDelegates::OnObjectPropertyChanged.RemoveAll(this);
+}
+
 FReply SCompositePassTree::OnKeyDown(const FGeometry& MyGeometry, const FKeyEvent& InKeyEvent)
 {
 	if (CommandList->ProcessCommandBindings(InKeyEvent))
@@ -1126,6 +1138,11 @@ ECheckBoxState SCompositePassTree::GetGlobalEnabledState() const
 		
 		for (const TObjectPtr<UCompositePassBase>& Pass : PassList)
 		{
+			if (!IsValid(Pass))
+			{
+				continue;
+			}
+			
 			if (Pass->IsEnabled())
 			{
 				bAnyPassesEnabled = true;
@@ -1402,6 +1419,27 @@ FSlateColor SCompositePassTree::GetFilterStatusColor() const
 	return FSlateColor::UseForeground();
 }
 
+void SCompositePassTree::OnObjectPropertyChanged(UObject* InObject, FPropertyChangedEvent& InPropertyChangedEvent)
+{
+	if (!Plate.IsValid())
+	{
+		return;
+	}
+
+	if (InObject != Plate.Get())
+	{
+		return;
+	}
+
+	const FName PropName = InPropertyChangedEvent.GetPropertyName();
+	if (PropName == GET_MEMBER_NAME_CHECKED(UCompositeLayerPlate, MediaPasses) ||
+		PropName == GET_MEMBER_NAME_CHECKED(UCompositeLayerPlate, ScenePasses) ||
+		PropName == GET_MEMBER_NAME_CHECKED(UCompositeLayerPlate, LayerPasses))
+	{
+		FillPassTreeItems();
+	}
+}
+
 void SCompositePassTree::CopySelectedItems()
 {
 	TArray<UObject*> ObjectsToCopy;
@@ -1545,6 +1583,9 @@ void SCompositePassTree::PasteSelectedItems()
 		
 		PassesToSelect.Add(PassTypeToPasteAt, TArray<int32>());
 		
+		FProperty* PassListProperty = GetPassListProperty(PassTypeToPasteAt);
+		Plate->PreEditChange(PassListProperty);
+		
 		for (UCompositePassBase* PassToPaste : Factory.CompositePasses)
 		{
 			UCompositePassBase* NewPass = DuplicateObject(PassToPaste, Plate.Get());
@@ -1554,6 +1595,9 @@ void SCompositePassTree::PasteSelectedItems()
 			++IndexToPasteAt;
 		}
 
+		FPropertyChangedEvent PropertyChangedEvent(PassListProperty, EPropertyChangeType::ArrayAdd);
+		Plate->PostEditChangeProperty(PropertyChangedEvent);
+		
 		constexpr bool bRefreshTree = false;
 		RefreshPassType(PassTypeToPasteAt, bRefreshTree);
 	}
@@ -1649,6 +1693,9 @@ void SCompositePassTree::DeleteSelectedItems()
 			// Delete the passes in reverse order so that indices don't get messed up as we delete
 			TArray<int32>& PassIndicesToDelete = PassesToDeleteByType.Value;
 			PassIndicesToDelete.Sort();
+
+			FProperty* PassListProperty = GetPassListProperty(PassType);
+			Plate->PreEditChange(PassListProperty);
 			
 			for (int32 Index = PassIndicesToDelete.Num() - 1; Index >= 0; --Index)
 			{
@@ -1657,6 +1704,9 @@ void SCompositePassTree::DeleteSelectedItems()
 
 			PassList.Shrink();
 			
+			FPropertyChangedEvent PropertyChangedEvent(PassListProperty, EPropertyChangeType::ArrayAdd);
+			Plate->PostEditChangeProperty(PropertyChangedEvent);
+			
 			constexpr bool bRefreshTreeView = false;
 			RefreshPassType(PassType, bRefreshTreeView);
 		}
diff --git a/Engine/Plugins/Experimental/Compositing/Composite/Source/CompositeEditor/Private/UI/SCompositePassTree.h b/Engine/Plugins/Experimental/Compositing/Composite/Source/CompositeEditor/Private/UI/SCompositePassTree.h
--- a/Engine/Plugins/Experimental/Compositing/Composite/Source/CompositeEditor/Private/UI/SCompositePassTree.h
+++ b/Engine/Plugins/Experimental/Compositing/Composite/Source/CompositeEditor/Private/UI/SCompositePassTree.h
@@ -62,6 +62,8 @@ public:
 
 	void Construct(const FArguments& InArgs, const TWeakObjectPtr<UCompositeLayerPlate>& InPlate);
 
+	virtual ~SCompositePassTree() override;
+	
 	// SWidget interface
 	virtual FReply OnKeyDown(const FGeometry& MyGeometry, const FKeyEvent& InKeyEvent) override;
 	// End of SWidget interface
@@ -116,8 +118,11 @@ private:
 
 	/** Gets status text color based on selection and active filter */
 	FSlateColor GetFilterStatusColor() const;
+
+	/** Raised when a property on the object has changed */
+	void OnObjectPropertyChanged(UObject* InObject, FPropertyChangedEvent& InPropertyChangedEvent);
 	
-	/** Generic command callbacks to handle copy, cut, paste, delete, and rename */
+	/** Generic command callbacks to handle copy, cut, paste, delete, and enable */
 	void CopySelectedItems();
 	bool CanCopySelectedItems();
 	void CutSelectedItems();
