Composite: Facilitate selection of plugin materials on the mesh component with an enum dropdown.

diff --git a/Engine/Plugins/Experimental/Compositing/Composite/Source/Composite/Private/Components/CompositeMeshComponent.cpp b/Engine/Plugins/Experimental/Compositing/Composite/Source/Composite/Private/Components/CompositeMeshComponent.cpp
--- a/Engine/Plugins/Experimental/Compositing/Composite/Source/Composite/Private/Components/CompositeMeshComponent.cpp
+++ b/Engine/Plugins/Experimental/Compositing/Composite/Source/Composite/Private/Components/CompositeMeshComponent.cpp
@@ -2,6 +2,7 @@
 
 #include "Components/CompositeMeshComponent.h"
 
+#include "Async/Async.h"
 #include "Engine/StaticMesh.h"
 #include "Materials/Material.h"
 #include "UObject/ConstructorHelpers.h"
@@ -14,8 +15,12 @@ UCompositeMeshComponent::UCompositeMeshComponent(const FObjectInitializer& Objec
 	struct FConstructorStatics
 	{
 		ConstructorHelpers::FObjectFinder<UStaticMesh> DefaultMesh;
+		ConstructorHelpers::FObjectFinder<UMaterialInterface> DefaultLitMaskedMaterial;
+		ConstructorHelpers::FObjectFinder<UMaterialInterface> DefaultUnlitAlphaCompositeMaterial;
 		FConstructorStatics()
 			: DefaultMesh(TEXT("/Composite/Meshes/SM_DefaultCompositeMesh.SM_DefaultCompositeMesh"))
+			, DefaultLitMaskedMaterial(TEXT("/Composite/Materials/M_CompositeMesh_Lit_Masked.M_CompositeMesh_Lit_Masked"))
+			, DefaultUnlitAlphaCompositeMaterial(TEXT("/Composite/Materials/M_CompositeMesh_Unlit_AlphaComposite.M_CompositeMesh_Unlit_AlphaComposite"))
 		{
 		}
 	};
@@ -26,6 +31,16 @@ UCompositeMeshComponent::UCompositeMeshComponent(const FObjectInitializer& Objec
 		SetStaticMesh(ConstructorStatics.DefaultMesh.Object);
 		SetCastShadow(false);
 	}
+
+	if (ensure(ConstructorStatics.DefaultLitMaskedMaterial.Object))
+	{
+		DefaultLitMaskedMaterial = ConstructorStatics.DefaultLitMaskedMaterial.Object;
+	}
+
+	if (ensure(ConstructorStatics.DefaultUnlitAlphaCompositeMaterial.Object))
+	{
+		DefaultUnlitAlphaCompositeMaterial = ConstructorStatics.DefaultUnlitAlphaCompositeMaterial.Object;
+	}
 }
 
 void UCompositeMeshComponent::PostInitProperties()
@@ -58,10 +73,76 @@ void UCompositeMeshComponent::OnUnregister()
 	Super::OnUnregister();
 }
 
+void UCompositeMeshComponent::SetMaterialType(ECompositeMeshMaterialType InMaterialType)
+{
+	MaterialType = InMaterialType;
+
+	switch (MaterialType)
+	{
+	case ECompositeMeshMaterialType::DefaultLitMasked:
+		for (int32 ElementIndex = 0; ElementIndex < GetMaterials().Num(); ++ElementIndex)
+		{
+			SetMaterial(ElementIndex, DefaultLitMaskedMaterial);
+		}
+		break;
+	case ECompositeMeshMaterialType::DefaultUnlitAlphaComposite:
+		for (int32 ElementIndex = 0; ElementIndex < GetMaterials().Num(); ++ElementIndex)
+		{
+			SetMaterial(ElementIndex, DefaultUnlitAlphaCompositeMaterial);
+		}
+		break;
+	case ECompositeMeshMaterialType::Custom:
+	default:
+		// no-op
+		break;
+	}
+}
+
+void UCompositeMeshComponent::CreateRenderState_Concurrent(FRegisterComponentContext* Context)
+{
+	Super::CreateRenderState_Concurrent(Context);
+
+	TWeakObjectPtr<UCompositeMeshComponent> WeakThis = this;
+
+	// Catch any material change here, and keep our material type enum up-to-date.
+	AsyncTask(ENamedThreads::GameThread, [WeakThis = MoveTemp(WeakThis)]()
+		{
+			TStrongObjectPtr<UCompositeMeshComponent> CompositeMesh = WeakThis.Pin();
+
+			if (CompositeMesh.IsValid())
+			{
+				UMaterialInterface* CurrentMaterial = CompositeMesh->GetMaterial(0);
+				if (IsValid(CurrentMaterial))
+				{
+					UMaterial* ParentMaterial = CurrentMaterial->GetBaseMaterial();
+					if (ParentMaterial == CompositeMesh->DefaultLitMaskedMaterial)
+					{
+						CompositeMesh->MaterialType = ECompositeMeshMaterialType::DefaultLitMasked;
+					}
+					else if (ParentMaterial == CompositeMesh->DefaultUnlitAlphaCompositeMaterial)
+					{
+						CompositeMesh->MaterialType = ECompositeMeshMaterialType::DefaultUnlitAlphaComposite;
+					}
+					else
+					{
+						CompositeMesh->MaterialType = ECompositeMeshMaterialType::Custom;
+					}
+				}
+			}
+		});
+}
+
 #if WITH_EDITOR
 void UCompositeMeshComponent::PostEditChangeProperty(FPropertyChangedEvent& PropertyChangedEvent)
 {
 	Super::PostEditChangeProperty(PropertyChangedEvent);
+
+	const FName PropertyName = PropertyChangedEvent.GetMemberPropertyName();
+
+	if (PropertyName == GET_MEMBER_NAME_CHECKED(ThisClass, MaterialType))
+	{
+		SetMaterialType(MaterialType);
+	}
 }
 #endif
 
diff --git a/Engine/Plugins/Experimental/Compositing/Composite/Source/Composite/Public/Components/CompositeMeshComponent.h b/Engine/Plugins/Experimental/Compositing/Composite/Source/Composite/Public/Components/CompositeMeshComponent.h
--- a/Engine/Plugins/Experimental/Compositing/Composite/Source/Composite/Public/Components/CompositeMeshComponent.h
+++ b/Engine/Plugins/Experimental/Compositing/Composite/Source/Composite/Public/Components/CompositeMeshComponent.h
@@ -9,6 +9,20 @@
 
 #define UE_API COMPOSITE_API
 
+/** Material selection type for convenience. */
+UENUM(BlueprintType)
+enum class ECompositeMeshMaterialType : uint8
+{
+	/** Default lit material. Perfect for catching shadows and/or reflections, but poor alpha edge quality due to masked alpha. */
+	DefaultLitMasked,
+	
+	/** Default unlit alpha composite (translucent pre-multiplied alpha) material. Perfect for . */
+	DefaultUnlitAlphaComposite,
+
+	/** Custom user material */
+	Custom
+};
+
 UCLASS(MinimalAPI, Blueprintable, ClassGroup = Composite, EditInlineNew)
 class UCompositeMeshComponent : public UStaticMeshComponent
 {
@@ -30,6 +44,32 @@ public:
 #if WITH_EDITOR
 	UE_API virtual void PostEditChangeProperty(struct FPropertyChangedEvent& PropertyChangedEvent) override;
 #endif //WITH_EDITOR
+
+	/** Get material type. */
+	UFUNCTION(BlueprintGetter)
+	ECompositeMeshMaterialType GetMaterialType() const { return MaterialType; }
+
+	/** Set material type. */
+	UFUNCTION(BlueprintSetter)
+	void SetMaterialType(ECompositeMeshMaterialType InMaterialType);
+
+protected:
+	//~ Begin UActorComponent interface
+	virtual void CreateRenderState_Concurrent(FRegisterComponentContext* Context) override;
+	//~ End UActorComponent interface
+
+private:
+	/** Material type selection for easier access to default plugin materials. */
+	UPROPERTY(EditAnywhere, BlueprintGetter = GetMaterialType, BlueprintSetter = SetMaterialType, Category = "Materials", meta = (AllowPrivateAccess))
+	ECompositeMeshMaterialType MaterialType;
+
+	/** Default masked material. */
+	UPROPERTY()
+	TObjectPtr<UMaterialInterface> DefaultLitMaskedMaterial;
+
+	/** Default alpha composite material. */
+	UPROPERTY()
+	TObjectPtr<UMaterialInterface> DefaultUnlitAlphaCompositeMaterial;
 };
 
 #undef UE_API
