Composite: Single light shadow layer, for fast shadows on media using percent-closer filtering.

* Minimal test scenes adds less than 0.2ms with a 2048x2048 shadow map.
* Optional shadow casting actor that are removed from the scene depth custom render pass to create a cleaner shadow matte.

diff --git a/Engine/Plugins/Experimental/Compositing/Composite/Shaders/Private/CompositeSingleLightShadow.usf b/Engine/Plugins/Experimental/Compositing/Composite/Shaders/Private/CompositeSingleLightShadow.usf
new file mode 100644
--- /dev/null
+++ b/Engine/Plugins/Experimental/Compositing/Composite/Shaders/Private/CompositeSingleLightShadow.usf
@@ -0,0 +1,87 @@
+// Copyright Epic Games, Inc. All Rights Reserved.
+
+#include "/Engine/Private/Common.ush"
+#include "/Engine/Private/ShadowFilteringCommon.ush"
+
+Texture2D       SceneDepthTex;
+SamplerState    SceneDepthSampler;
+
+Texture2D       ShadowDepthTexture;
+SamplerState    ShadowDepthTextureSampler;
+
+float4          SceneDepthPositionScaleBias;
+float4x4        SceneDepthScreenToWorld;
+
+float4x4        ShadowMatrix;
+float4          ShadowBufferSize;
+float4          ShadowInvDeviceZToWorldZ;
+float4          ShadowStrengthBias;
+
+uint            bIsPerspectiveProjection;
+
+// Returns the screen position for projection matrix calculations depending on the type of projection the view is using
+float2 GetScreenPositionForProj(in float2 ScreenPosition, in float SceneDepth)
+{
+    return select(bIsPerspectiveProjection == 0, ScreenPosition, ScreenPosition * SceneDepth);
+}
+
+// Calculate world position from scene depth
+float4 CalcWorldFromSceneDepth(in float2 UV, in float Depth)
+{
+    float2 ScreenPosition = (UV - SceneDepthPositionScaleBias.wz) / SceneDepthPositionScaleBias.xy;
+    float4 HomogeneousWorldPosition = mul(float4(GetScreenPositionForProj(ScreenPosition, Depth), Depth, 1), SceneDepthScreenToWorld);
+    float3 WorldPosition = HomogeneousWorldPosition.xyz / HomogeneousWorldPosition.w;
+
+    return float4(WorldPosition, 1.0);
+}
+
+// Convert device-z depth to scene depth
+float ShadowDeviceZToSceneDepth(float DeviceZ)
+{
+    return DeviceZ * ShadowInvDeviceZToWorldZ[0] + ShadowInvDeviceZToWorldZ[1] + 1.0f / (DeviceZ * ShadowInvDeviceZToWorldZ[2] - ShadowInvDeviceZToWorldZ[3]);
+}
+
+// Percentage-close filtering shadow mapping
+float ManualPCF(in float4 ShadowPosition)
+{
+    // TODO: Improve shadow bias logic with normal & resolution info
+    float SceneDepth = ShadowDeviceZToSceneDepth(ShadowPosition.z) - ShadowStrengthBias.y;
+
+    FPCFSamplerSettings Settings;
+    Settings.ShadowDepthTexture = ShadowDepthTexture;
+    Settings.ShadowDepthTextureSampler = ShadowDepthTextureSampler;
+    Settings.ShadowBufferSize = ShadowBufferSize;
+    Settings.TransitionScale = 1.0; // TODO: Parametrize like FProjectedShadowInfo::ComputeTransitionSize
+    Settings.SceneDepth = SceneDepth;
+    Settings.bSubsurface = false;
+    Settings.bTreatMaxDepthUnshadowed = false;
+    Settings.DensityMulConstant = 0;
+    Settings.ProjectionDepthBiasParameters = 0;
+    float Shadow = Manual5x5PCF(ShadowPosition.xy, Settings);
+
+    return Shadow;
+}
+
+void MainPS(
+    noperspective float4 InUVAndScreenPos : TEXCOORD0,
+    out float4 OutColor : SV_Target0)
+{
+    float2 UV = InUVAndScreenPos.xy;
+    float SceneDepth = SceneDepthTex.SampleLevel(SceneDepthSampler, UV, 0).r;
+    
+    float4 WorldPosition = CalcWorldFromSceneDepth(UV, SceneDepth);
+    float4 ShadowPosition = mul(WorldPosition, ShadowMatrix);
+    float2 ShadowUV = ShadowPosition.xy;
+    
+    if(any(ShadowUV < 0.0) || any(ShadowUV > 1.0))
+    {
+        OutColor.rgb = 1.0;
+    }
+    else
+    {
+        float ShadowVisibility = ManualPCF(ShadowPosition);
+        OutColor.rgb = lerp(1.0, ShadowVisibility, ShadowStrengthBias.x);
+    }
+
+    OutColor.a = 1.0;
+}
diff --git a/Engine/Plugins/Experimental/Compositing/Composite/Source/Composite/Private/Layers/CompositeLayerSingleLightShadow.cpp b/Engine/Plugins/Experimental/Compositing/Composite/Source/Composite/Private/Layers/CompositeLayerSingleLightShadow.cpp
new file mode 100644
--- /dev/null
+++ b/Engine/Plugins/Experimental/Compositing/Composite/Source/Composite/Private/Layers/CompositeLayerSingleLightShadow.cpp
@@ -0,0 +1,345 @@
+// Copyright Epic Games, Inc. All Rights Reserved.
+
+#include "Layers/CompositeLayerSingleLightShadow.h"
+
+#include "Camera/CameraComponent.h"
+#include "CompositeActor.h"
+#include "CompositeRenderTargetPool.h"
+#include "Passes/CompositePassMaterial.h"
+#include "Components/SceneCaptureComponent2D.h"
+#include "Engine/DirectionalLight.h"
+#include "Engine/TextureRenderTarget2D.h"
+#include "Passes/CompositeCorePassProxy.h"
+#include "RHI.h"
+
+DECLARE_GPU_STAT_NAMED(FCompositeSingleLightShadow, TEXT("Composite.SingleLayerShadow"));
+
+class FCompositeSingleLightShadowShader : public FGlobalShader
+{
+	DECLARE_GLOBAL_SHADER(FCompositeSingleLightShadowShader);
+	SHADER_USE_PARAMETER_STRUCT(FCompositeSingleLightShadowShader, FGlobalShader);
+
+	BEGIN_SHADER_PARAMETER_STRUCT(FParameters, )
+		SHADER_PARAMETER_RDG_TEXTURE(Texture2D, SceneDepthTex)
+		SHADER_PARAMETER_SAMPLER(SamplerState, SceneDepthSampler)
+
+		SHADER_PARAMETER_RDG_TEXTURE(Texture2D, ShadowDepthTexture)
+		SHADER_PARAMETER_SAMPLER(SamplerState, ShadowDepthTextureSampler)
+
+		SHADER_PARAMETER(FVector4f, SceneDepthPositionScaleBias)
+		SHADER_PARAMETER(FMatrix44f, SceneDepthScreenToWorld)
+		SHADER_PARAMETER(FMatrix44f, ShadowMatrix)
+		SHADER_PARAMETER(FVector4f, ShadowBufferSize)
+		SHADER_PARAMETER(FVector4f, ShadowInvDeviceZToWorldZ)
+		SHADER_PARAMETER(FVector4f, ShadowStrengthBias)
+		SHADER_PARAMETER(uint32, bIsPerspectiveProjection)
+
+		RENDER_TARGET_BINDING_SLOTS()
+	END_SHADER_PARAMETER_STRUCT()
+};
+IMPLEMENT_GLOBAL_SHADER(FCompositeSingleLightShadowShader, "/Plugin/Composite/Private/CompositeSingleLightShadow.usf", "MainPS", SF_Pixel);
+
+
+namespace UE
+{
+	namespace CompositeCore
+	{
+		/**
+		* Utility function to create the inverse depth projection transform to be used by the shader system.
+		* Copied from SceneView.cpp.
+		*/
+		FVector4f CreateInvDeviceZToWorldZTransform(const FMatrix& ProjMatrix)
+		{
+			float DepthMul = (float)ProjMatrix.M[2][2];
+			float DepthAdd = (float)ProjMatrix.M[3][2];
+
+			if (DepthAdd == 0.f)
+			{
+				// Avoid dividing by 0 in this case
+				DepthAdd = 0.00000001f;
+			}
+
+			bool bIsPerspectiveProjection = ProjMatrix.M[3][3] < 1.0f;
+			if (bIsPerspectiveProjection)
+			{
+				float SubtractValue = DepthMul / DepthAdd;
+
+				// Subtract a tiny number to avoid divide by 0 errors in the shader when a very far distance is decided from the depth buffer.
+				// This fixes fog not being applied to the black background in the editor.
+				SubtractValue -= 0.00000001f;
+
+				return FVector4f( 0.0f, 0.0f, 1.0f / DepthAdd, SubtractValue );
+			}
+			else
+			{
+				return FVector4f((float)(1.0f / ProjMatrix.M[2][2]), (float)(-ProjMatrix.M[3][2] / ProjMatrix.M[2][2] + 1.0f), 0.0f, 1.0f);
+			}
+		}
+
+		class FCompositeSingleLightShadowProxy : public FCompositeCorePassProxy
+		{
+		public:
+			using FCompositeCorePassProxy::FCompositeCorePassProxy;
+
+			FPassTexture Add(FRDGBuilder& GraphBuilder, const FSceneView& InView, const FPassInputArray& Inputs, const FPassContext& PassContext) const override
+			{
+				RDG_EVENT_SCOPE_STAT(GraphBuilder, FCompositeSingleLightShadow, "Composite.SingleLayerShadow");
+				RDG_GPU_STAT_SCOPE(GraphBuilder, FCompositeSingleLightShadow);
+
+				check(ValidateInputs(Inputs));
+
+				const UE::CompositeCore::FResourceMetadata& Metadata = Inputs[0].Metadata;
+				const FScreenPassTexture& SceneDepth = Inputs[0].Texture;
+				const FScreenPassTexture& ShadowMap = Inputs[1].Texture;
+				FScreenPassRenderTarget Output = Inputs.OverrideOutput;
+
+				if (!Output.IsValid())
+				{
+					FRDGTextureDesc Desc = SceneDepth.Texture->Desc;
+					Desc.Format = FSceneTexturesConfig::Get().ColorFormat;
+					Output = CreateOutputRenderTarget(GraphBuilder, InView, PassContext.OutputViewRect, Desc, TEXT("CompositeSingleLayerShadow"));
+				}
+
+				// Calculate scene depth position scale bias
+				const FIntPoint Resolution = FIntPoint(SceneDepth.ViewRect.Size().X, SceneDepth.ViewRect.Size().Y);
+				FVector4f SceneDepthPositionScaleBias = FSceneView::GetScreenPositionScaleBias(Resolution, FIntRect(FIntPoint::ZeroValue, Resolution));
+
+				// Calculate scene depth screen to world matrix (without jitter).
+				FMatrix44f SceneDepthScreenToWorld = FMatrix44f(
+					InView.ViewMatrices.GetScreenToClipMatrix() *
+					InView.ViewMatrices.GetProjectionNoAAMatrix().Inverse() *
+					InView.ViewMatrices.GetInvViewMatrix()
+				);
+
+				// The shadow map texture size
+				const FIntVector ShadowMapSize = ShadowMap.Texture->Desc.GetSize();
+
+				FGlobalShaderMap* GlobalShaderMap = GetGlobalShaderMap(InView.GetFeatureLevel());
+				FScreenPassTextureViewport Viewport = FScreenPassTextureViewport(Output);
+
+				FCompositeSingleLightShadowShader::FParameters* Parameters = GraphBuilder.AllocParameters<FCompositeSingleLightShadowShader::FParameters>();
+				Parameters->SceneDepthTex				= SceneDepth.Texture;
+				Parameters->SceneDepthSampler			= TStaticSamplerState<SF_Bilinear>::GetRHI();
+				Parameters->ShadowDepthTexture			= ShadowMap.Texture;
+				Parameters->ShadowDepthTextureSampler	= TStaticSamplerState<SF_Point>::GetRHI();
+				Parameters->SceneDepthPositionScaleBias	= SceneDepthPositionScaleBias;
+				Parameters->SceneDepthScreenToWorld		= SceneDepthScreenToWorld;
+				Parameters->ShadowInvDeviceZToWorldZ	= ShadowInvDeviceZToWorldZ;
+				Parameters->ShadowMatrix				= ShadowMatrix;
+				Parameters->ShadowBufferSize			= FVector4f((float)ShadowMapSize.X, (float)ShadowMapSize.Y, 1.0f / (float)ShadowMapSize.X, 1.0f / (float)ShadowMapSize.Y);
+				Parameters->ShadowStrengthBias			= ShadowStrengthBias;
+				// We can use the current view since we the depth render matches the main view
+				Parameters->bIsPerspectiveProjection	= InView.IsPerspectiveProjection();
+				Parameters->RenderTargets[0]			= Output.GetRenderTargetBinding();
+
+				TShaderMapRef<FCompositeSingleLightShadowShader> PixelShader(GlobalShaderMap);
+				AddDrawScreenPass(
+					GraphBuilder,
+					RDG_EVENT_NAME("Composite.SingleLayerShadow (%dx%d)", Viewport.Extent.X, Viewport.Extent.Y),
+					InView,
+					Viewport,
+					Viewport,
+					PixelShader,
+					Parameters
+				);
+
+				return UE::CompositeCore::FPassTexture{ MoveTemp(Output), Metadata };
+			}
+
+			FMatrix44f ShadowMatrix;
+			FVector4f ShadowInvDeviceZToWorldZ;
+			FVector4f ShadowStrengthBias;
+		};
+	}
+}
+
+UCompositeLayerSingleLightShadow::UCompositeLayerSingleLightShadow(const FObjectInitializer& ObjectInitializer)
+	: Super(ObjectInitializer)
+{
+	Operation = ECompositeCoreMergeOp::Multiply;
+}
+
+void UCompositeLayerSingleLightShadow::BeginDestroy()
+{
+	Super::BeginDestroy();
+
+	if (USceneCaptureComponent2D* CaptureComponent = CachedShadowMapCapture.Get())
+	{
+		FCompositeRenderTargetPool::Get().ReleaseTarget(CaptureComponent->TextureTarget);
+	}
+
+	if (USceneCaptureComponent2D* CaptureComponent = CachedSceneDepthCapture.Get())
+	{
+		FCompositeRenderTargetPool::Get().ReleaseTarget(CaptureComponent->TextureTarget);
+	}
+}
+
+#if WITH_EDITOR
+void UCompositeLayerSingleLightShadow::PostEditChangeProperty(FPropertyChangedEvent& PropertyChangedEvent)
+{
+	Super::PostEditChangeProperty(PropertyChangedEvent);
+
+	const FName PropertyName = PropertyChangedEvent.GetMemberPropertyName();
+
+	if (PropertyName == GET_MEMBER_NAME_CHECKED(ThisClass, bIsEnabled))
+	{
+		ACompositeActor* CompositeActor = GetTypedOuter<ACompositeActor>();
+		if (IsValid(CompositeActor))
+		{
+			CompositeActor->ActivateSceneCaptures(this, bIsEnabled);
+		}
+	}
+	else if (PropertyName == GET_MEMBER_NAME_CHECKED(ThisClass, ShadowCastingActors))
+	{
+		SetShadowCastingActors(ShadowCastingActors);
+	}
+}
+#endif
+
+bool UCompositeLayerSingleLightShadow::GetProxy(FTraversalContext& InContext, FSceneRenderingBulkObjectAllocator& InFrameAllocator, FCompositeCorePassProxy*& OutProxy) const
+{
+	using namespace UE::CompositeCore;
+
+	ACompositeActor* CompositeActor = GetTypedOuter<ACompositeActor>();
+	if (!IsValid(CompositeActor) || !IsValid(Light))
+	{
+		return false;
+	}
+
+	// Disable with zero strength
+	if (ShadowStrength <= 0.0f)
+	{
+		return false;
+	}
+
+	// Create new shadow map render scene capture (custom render pass) if needed
+	if (!CachedShadowMapCapture.IsValid())
+	{
+		CachedShadowMapCapture = CompositeActor->FindOrCreateSceneCapture<USceneCaptureComponent2D>(this, 0, FName("ShadowMapCapture"));
+		CachedShadowMapCapture->PrimitiveRenderMode = ESceneCapturePrimitiveRenderMode::PRM_RenderScenePrimitives;
+		CachedShadowMapCapture->CaptureSource = SCS_SceneDepth;
+		CachedShadowMapCapture->bRenderInMainRenderer = true;
+		CachedShadowMapCapture->bMainViewFamily = false;
+		CachedShadowMapCapture->bMainViewResolution = false;
+		CachedShadowMapCapture->bMainViewCamera = false;
+		CachedShadowMapCapture->bInheritMainViewCameraPostProcessSettings = false;
+		CachedShadowMapCapture->bUseRayTracingIfEnabled = false;
+		CachedShadowMapCapture->ProjectionType = ECameraProjectionMode::Orthographic;
+	}
+
+	// Create new scene depth render without jitter
+	if (!CachedSceneDepthCapture.IsValid())
+	{
+		CachedSceneDepthCapture = CompositeActor->FindOrCreateSceneCapture<USceneCaptureComponent2D>(this, 1, FName("SceneDepthCapture"));
+		CachedSceneDepthCapture->PrimitiveRenderMode = ESceneCapturePrimitiveRenderMode::PRM_RenderScenePrimitives;
+		CachedSceneDepthCapture->CaptureSource = SCS_SceneDepth;
+		CachedSceneDepthCapture->bRenderInMainRenderer = true;
+		CachedSceneDepthCapture->bIgnoreScreenPercentage = true;
+		CachedSceneDepthCapture->HiddenActors = ShadowCastingActors;
+	}
+
+	TStrongObjectPtr<USceneCaptureComponent2D> ShadowMapCapture = CachedShadowMapCapture.Pin();
+	TStrongObjectPtr<USceneCaptureComponent2D> SceneDepthCapture = CachedSceneDepthCapture.Pin();
+
+	// Update properties that may be dynamic
+	{
+		// Keep exposed properties & transform in sync
+		ShadowMapCapture->OrthoWidth = OrthographicWidth;
+		ShadowMapCapture->SetWorldTransform(Light->GetTransform());
+
+		// Acquire render target
+		FCompositeRenderTargetPool::Get().ConditionalAcquireTarget(ShadowMapCapture.Get(), ShadowMapCapture->TextureTarget, FIntPoint(ShadowMapResolution, ShadowMapResolution), RTF_R32f);
+		FCompositeRenderTargetPool::Get().ConditionalAcquireTarget(SceneDepthCapture.Get(), SceneDepthCapture->TextureTarget, GetRenderResolution(), RTF_R32f);
+	}
+	
+	// Setup child shadow pass proxy
+	FCompositeSingleLightShadowProxy* ShadowProxy;
+	{
+		// Display-resolution scene depth without jitter as an external texture.
+		InContext.ExternalTextures.Add(FExternalTexture{ SceneDepthCapture->TextureTarget, FResourceMetadata{} });
+		const int32 SceneDepthIndex = InContext.ExternalTextures.Num() - 1;
+
+		// Shadow map texture
+		InContext.ExternalTextures.Add(FExternalTexture{ ShadowMapCapture->TextureTarget, FResourceMetadata{} });
+		const int32 ShadowMapIndex = InContext.ExternalTextures.Num() - 1;
+
+		// Setup shadow pass inputs
+		FPassInputDeclArray PassInputs;
+		PassInputs.SetNum(2);
+		PassInputs[0].Set<FPassExternalResourceDesc>({ EXTERNAL_RANGE_START_ID + SceneDepthIndex });
+		PassInputs[1].Set<FPassExternalResourceDesc>({ EXTERNAL_RANGE_START_ID + ShadowMapIndex });
+
+		ShadowProxy = InFrameAllocator.Create<FCompositeSingleLightShadowProxy>(PassInputs);
+		ShadowProxy->ShadowStrengthBias = FVector4f(ShadowStrength, ShadowBias, 0.0f, 0.0f);
+		GetShadowMatrices(*ShadowMapCapture, ShadowProxy->ShadowMatrix, ShadowProxy->ShadowInvDeviceZToWorldZ);
+	}
+
+	FPassInputDeclArray Inputs;
+	Inputs.SetNum(FixedNumLayerInputs);
+	Inputs[0].Set<const FCompositeCorePassProxy*>(ShadowProxy);
+	Inputs[1] = GetDefaultSecondInput(InContext);
+
+	OutProxy = InFrameAllocator.Create<FMergePassProxy>(MoveTemp(Inputs), GetMergeOperation(InContext), TEXT("SingleLightShadowCatcher"));
+	return true;
+}
+
+const TArray<AActor*> UCompositeLayerSingleLightShadow::GetShadowCastingActors() const
+{
+	return ShadowCastingActors;
+}
+
+void UCompositeLayerSingleLightShadow::SetShadowCastingActors(TArray<AActor*> InShadowCaster)
+{
+	ShadowCastingActors = InShadowCaster;
+
+	TStrongObjectPtr<USceneCaptureComponent2D> SceneDepthCapture = CachedSceneDepthCapture.Pin();
+	if (SceneDepthCapture.IsValid())
+	{
+		SceneDepthCapture->HiddenActors = ShadowCastingActors;
+	}
+}
+
+void UCompositeLayerSingleLightShadow::GetShadowMatrices(USceneCaptureComponent2D& ShadowMapCapture, FMatrix44f& OutShadowMatrix, FVector4f& OutShadowInvDeviceZToWorldZ) const
+{
+	FMinimalViewInfo ViewInfo;
+	ShadowMapCapture.GetCameraView(GetWorld()->DeltaTimeSeconds, ViewInfo);
+
+	// Ortho clip plane values from SceneCaptureRendering.cpp
+	ViewInfo.OrthoNearClipPlane = 0;
+	ViewInfo.OrthoFarClipPlane = UE_FLOAT_HUGE_DISTANCE / 4.0f;
+
+	FMatrix ViewRotationMatrix = FInverseRotationMatrix(ViewInfo.Rotation) * FMatrix(
+		FPlane(0, 0, 1, 0),
+		FPlane(1, 0, 0, 0),
+		FPlane(0, 1, 0, 0),
+		FPlane(0, 0, 0, 1));
+
+	FVector LocalViewOrigin = ViewInfo.Location;
+	if (!ViewRotationMatrix.GetOrigin().IsNearlyZero(0.0f))
+	{
+		LocalViewOrigin += ViewRotationMatrix.InverseTransformPosition(FVector::ZeroVector);
+		ViewRotationMatrix = ViewRotationMatrix.RemoveTranslation();
+	}
+	const FMatrix ShadowViewMatrix = FTranslationMatrix(-LocalViewOrigin) * ViewRotationMatrix;
+	const FMatrix ShadowProjectionMatrix = AdjustProjectionMatrixForRHI(ViewInfo.CalculateProjectionMatrix());
+	
+	//TODO: Max subject depth logic will need to be updated once support perspective projection.
+	float MaxSubjectDepth = 1.0f;
+	float InvMaxSubjectDepth = 1.0f / MaxSubjectDepth;
+	const float ShadowResolutionFractionX = 0.5f;
+	const float ShadowResolutionFractionY = 0.5f;
+
+	// Note: See FProjectedShadowInfo::GetWorldToShadowMatrix as a reference.
+	const FMatrix WorldToShadowMatrix = ShadowViewMatrix * ShadowProjectionMatrix *
+		FMatrix(
+			FPlane(ShadowResolutionFractionX, 0, 0, 0),
+			FPlane(0, -ShadowResolutionFractionY, 0, 0),
+			FPlane(0, 0, InvMaxSubjectDepth, 0),
+			FPlane(ShadowResolutionFractionX, ShadowResolutionFractionY, 0, 1)
+		);
+
+	OutShadowMatrix = FMatrix44f(WorldToShadowMatrix);
+	OutShadowInvDeviceZToWorldZ = UE::CompositeCore::CreateInvDeviceZToWorldZTransform(ShadowProjectionMatrix);
+}
+
+
diff --git a/Engine/Plugins/Experimental/Compositing/Composite/Source/Composite/Public/Layers/CompositeLayerSingleLightShadow.h b/Engine/Plugins/Experimental/Compositing/Composite/Source/Composite/Public/Layers/CompositeLayerSingleLightShadow.h
new file mode 100644
--- /dev/null
+++ b/Engine/Plugins/Experimental/Compositing/Composite/Source/Composite/Public/Layers/CompositeLayerSingleLightShadow.h
@@ -0,0 +1,109 @@
+// Copyright Epic Games, Inc. All Rights Reserved.
+
+#pragma once
+
+#include "EngineDefines.h"
+#include "CompositeLayerBase.h"
+#include "UObject/WeakObjectPtr.h"
+
+#include "CompositeLayerSingleLightShadow.generated.h"
+
+/**
+* Note that much of the custom logic here could be removed given access to the engine's CSM/VSM built-in functionality.
+* However, those techniques remain out-of-reach to plugins since they are renderer-private.
+*/
+
+#define UE_API COMPOSITE_API
+
+class AActor;
+class ALight;
+class USceneCaptureComponent2D;
+
+/**
+* Layer for catching shadows from a single light.
+* 
+* Classic shadow map percentage-closer filtering technique (without cascades), defaulting to 5x5.
+* Implemented with custom render passes for minimal performance cost.
+*/
+UCLASS(MinimalAPI, BlueprintType, Blueprintable, EditInlineNew, CollapseCategories, meta = (DisplayName = "Single Light Shadow"))
+class UCompositeLayerSingleLightShadow : public UCompositeLayerBase
+{
+	GENERATED_BODY()
+
+public:
+	/** Constructor. */
+	UE_API UCompositeLayerSingleLightShadow(const FObjectInitializer& ObjectInitializer);
+	
+	/** Destructor. */
+	UE_API virtual ~UCompositeLayerSingleLightShadow() = default;
+
+	UE_API virtual void BeginDestroy() override;
+
+#if WITH_EDITOR
+	UE_API virtual void PostEditChangeProperty(FPropertyChangedEvent& PropertyChangedEvent) override;
+#endif
+
+	/** Getter function to override, returning pass layer proxies to be passed to the render thread. (Proxy objects should be allocated from the provided allocator.) */
+	UE_API virtual bool GetProxy(FTraversalContext& InContext, FSceneRenderingBulkObjectAllocator& InFrameAllocator, FCompositeCorePassProxy*& OutProxy) const override;
+
+public:
+	/**
+	* The reference light, whose exact transform will by used to render a shadow map.
+	* Please make sure it is aligned to your shadow casting scene by piloting the light.
+	* 
+	* Only directional lights are currently supported.
+	*/
+	UPROPERTY(EditAnywhere, BlueprintReadWrite, Category = "", meta=(AllowedClasses="/Script/Engine.DirectionalLight"))
+	TObjectPtr<ALight> Light;
+
+	/** Get shadow caster actors. */
+	UFUNCTION(BlueprintGetter)
+	UE_API const TArray<AActor*> GetShadowCastingActors() const;
+
+	/** Set shadow caster actors. */
+	UFUNCTION(BlueprintSetter)
+	UE_API void SetShadowCastingActors(TArray<AActor*> InShadowCaster);
+
+	/**
+	* The desired width (in world units) of the shadow map view (ignored if light is not directional).
+	* Reduce size to only contain the shadow casting geometry in view of the light.
+	*/
+	UPROPERTY(EditAnywhere, BlueprintReadWrite, Category = "")
+	float OrthographicWidth = DEFAULT_ORTHOWIDTH;
+
+	/**
+	* Resolution of the shadow map texture. Lower resolutions will cause more blurry shadows, but require a higher shadow bias.
+	*/
+	UPROPERTY(EditAnywhere, BlueprintReadWrite, Category = "", meta = (ClampMin = "128", UIMin = "128"))
+	int32 ShadowMapResolution = 2048;
+
+	/** Basic shadow strength multiplier, from 0 to 1. */
+	UPROPERTY(EditAnywhere, BlueprintReadWrite, Category = "", meta = (ClampMin = "0.0", ClampMax = "1.0", UIMin = "0.0", UIMax = "1.0"))
+	float ShadowStrength = 0.5f;
+
+	/**
+	* Shadow bias (in world units) to avoid shadow acne. Increase until unwanted shadowing disappears.
+	*/
+	UPROPERTY(EditAnywhere, BlueprintReadWrite, Category = "", meta = (ClampMin = "0.0", UIMin = "0.0"))
+	float ShadowBias = 2.0f;
+
+private:
+	/**
+	* Optional: List of shadow casting actors to be removed from the main camera view's scene depth.
+	* This helps create a cleaner shadow matte preventing aliased dark edges around CG objects.
+	*/
+	UPROPERTY(EditAnywhere, BlueprintGetter = GetShadowCastingActors, BlueprintSetter = SetShadowCastingActors, Category = "", meta = (AllowPrivateAccess = true))
+	TArray<TObjectPtr<AActor>> ShadowCastingActors;
+
+	/** Convenience function to calculate shadow matrices. */
+	void GetShadowMatrices(USceneCaptureComponent2D& ShadowMapCapture, FMatrix44f& OutShadowMatrix, FVector4f& OutShadowInvDeviceZToWorldZ) const;
+
+	/** Cached pointer to the composite-managed shadow map scene capture. */
+	mutable TWeakObjectPtr<USceneCaptureComponent2D> CachedShadowMapCapture;
+
+	/** Cached pointer to the scene depth capture. */
+	mutable TWeakObjectPtr<USceneCaptureComponent2D> CachedSceneDepthCapture;
+};
+
+#undef UE_API
+
