Composure: Added media profile source picker to the Plate's Texture property. Plates will also attempt to open the active media profile's media source if their Texture is pointing to a media profile default media texture

diff --git a/Engine/Plugins/Experimental/Compositing/Composite/Composite.uplugin b/Engine/Plugins/Experimental/Compositing/Composite/Composite.uplugin
--- a/Engine/Plugins/Experimental/Compositing/Composite/Composite.uplugin
+++ b/Engine/Plugins/Experimental/Compositing/Composite/Composite.uplugin
@@ -39,6 +39,10 @@
 		{
 			"Name": "CameraCalibrationCore",
 			"Enabled": true
+		},
+		{
+			"Name": "MediaFrameworkUtilities",
+			"Enabled": true
 		}
 	]
 }
\ No newline at end of file
diff --git a/Engine/Plugins/Experimental/Compositing/Composite/Source/Composite/Composite.Build.cs b/Engine/Plugins/Experimental/Compositing/Composite/Source/Composite/Composite.Build.cs
--- a/Engine/Plugins/Experimental/Compositing/Composite/Source/Composite/Composite.Build.cs
+++ b/Engine/Plugins/Experimental/Compositing/Composite/Source/Composite/Composite.Build.cs
@@ -30,7 +30,8 @@ public class Composite : ModuleRules
 				"RenderCore",
 				"Renderer",
 				"RHI",
-				"CameraCalibrationCore"
+				"CameraCalibrationCore",
+				"MediaFrameworkUtilities"
 			}
 		);
 
diff --git a/Engine/Plugins/Experimental/Compositing/Composite/Source/Composite/Private/Layers/CompositeLayerPlate.cpp b/Engine/Plugins/Experimental/Compositing/Composite/Source/Composite/Private/Layers/CompositeLayerPlate.cpp
--- a/Engine/Plugins/Experimental/Compositing/Composite/Source/Composite/Private/Layers/CompositeLayerPlate.cpp
+++ b/Engine/Plugins/Experimental/Compositing/Composite/Source/Composite/Private/Layers/CompositeLayerPlate.cpp
@@ -14,6 +14,9 @@
 #include "Materials/Material.h"
 #include "Materials/MaterialInstanceDynamic.h"
 #include "Passes/CompositePassDistortion.h"
+#include "Profile/IMediaProfileManager.h"
+#include "Profile/MediaProfile.h"
+#include "Profile/MediaProfilePlaybackManager.h"
 #include "UObject/ConstructorHelpers.h"
 
 #if WITH_EDITOR
@@ -203,6 +206,7 @@ void UCompositeLayerPlate::PostLoad()
 	if (!HasAnyFlags(RF_ClassDefaultObject))
 	{
 		UpdateCompositeMeshes();
+		TryOpenMediaProfileSource();
 	}
 }
 
@@ -213,15 +217,17 @@ void UCompositeLayerPlate::BeginDestroy()
 	FCompositeRenderTargetPool::Get().ReleaseTarget(MediaRenderTarget);
 	FCompositeRenderTargetPool::Get().ReleaseTarget(SceneRenderTarget);
 
-#if WITH_EDITOR
 	// If we are not the class default object...
 	if (!HasAnyFlags(RF_ClassDefaultObject))
 	{
+		TryCloseMediaProfileSource();
+		
+#if WITH_EDITOR
 		// Hook into pre/post save.
 		FEditorDelegates::PreSaveWorldWithContext.RemoveAll(this);
 		FEditorDelegates::PostSaveWorldWithContext.RemoveAll(this);
-	}
 #endif // WITH_EDITOR
+	}
 }
 
 const TArray<AActor*> UCompositeLayerPlate::GetCompositeMeshes() const
@@ -292,6 +298,10 @@ void UCompositeLayerPlate::PreEditChange(FProperty* PropertyThatWillChange)
 			Subsystem->UnregisterPrimitives(GetPrimitives());
 		}
 	}
+	else if (PropertyName == GET_MEMBER_NAME_CHECKED(ThisClass, Texture))
+	{
+		TryCloseMediaProfileSource();
+	}
 }
 
 void UCompositeLayerPlate::PostEditChangeProperty(FPropertyChangedEvent& PropertyChangedEvent)
@@ -345,6 +355,10 @@ void UCompositeLayerPlate::PostEditChangeProperty(FPropertyChangedEvent& Propert
 
 		PreEditCompositeMeshes.Reset();
 	}
+	else if (PropertyName == GET_MEMBER_NAME_CHECKED(ThisClass, Texture))
+	{
+		TryOpenMediaProfileSource();
+	}
 
 	// For simplicity, we make sure to update composite meshes on any property change.
 	UpdateCompositeMeshes();
@@ -360,6 +374,15 @@ void UCompositeLayerPlate::PostTransacted(const FTransactionObjectEvent& InTrans
 
 		if( ChangedProperties.Contains(GET_MEMBER_NAME_CHECKED(ThisClass, Texture)))
 		{
+			if (InTransactionEvent.GetEventType() == ETransactionObjectEventType::UndoRedo)
+			{
+				// Special processing for undo/redo as PreEdit and PostEdit are not called during undo/redo
+				// Attempt to close any outstanding media profile sources this plate may have opened for Texture,
+				// and then attempt to open the corresponding media source if Texture is a media profile media texture
+				TryCloseMediaProfileSource();
+				TryOpenMediaProfileSource();
+			}
+			
 			UpdateCompositeMeshes();
 		}
 		else if (ChangedProperties.Contains(GET_MEMBER_NAME_CHECKED(ThisClass, CompositeMeshes)))
@@ -576,6 +599,41 @@ int32 UCompositeLayerPlate::GetValidPassesNum(TArrayView<const TObjectPtr<UCompo
 	return Count;
 }
 
+void UCompositeLayerPlate::TryOpenMediaProfileSource()
+{
+	UMediaTexture* MediaTexture = Cast<UMediaTexture>(Texture);
+	if (!MediaTexture)
+	{
+		return;
+	}
+	
+	UMediaProfile* ActiveMediaProfile = IMediaProfileManager::Get().GetCurrentMediaProfile();
+	if (!ActiveMediaProfile)
+	{
+		return;
+	}
+	
+	int32 MediaSourceIndex = INDEX_NONE;
+	if (ActiveMediaProfile->GetPlaybackManager()->IsValidSourceMediaTexture(MediaTexture, MediaSourceIndex))
+	{
+		ActiveMediaProfile->GetPlaybackManager()->OpenSourceFromIndex(MediaSourceIndex, this);
+	}
+}
+
+void UCompositeLayerPlate::TryCloseMediaProfileSource()
+{
+	UMediaProfile* ActiveMediaProfile = IMediaProfileManager::Get().GetCurrentMediaProfile();
+	if (!ActiveMediaProfile)
+	{
+		return;
+	}
+
+	UMediaProfilePlaybackManager::FCloseSourceArgs Args;
+	Args.Consumer = this;
+				
+	ActiveMediaProfile->GetPlaybackManager()->CloseSourcesForConsumer(Args);
+}
+
 #if WITH_EDITOR
 void UCompositeLayerPlate::OnPreSaveWorld(UWorld* InWorld, FObjectPreSaveContext ObjectSaveContext)
 {
diff --git a/Engine/Plugins/Experimental/Compositing/Composite/Source/Composite/Public/Layers/CompositeLayerPlate.h b/Engine/Plugins/Experimental/Compositing/Composite/Source/Composite/Public/Layers/CompositeLayerPlate.h
--- a/Engine/Plugins/Experimental/Compositing/Composite/Source/Composite/Public/Layers/CompositeLayerPlate.h
+++ b/Engine/Plugins/Experimental/Compositing/Composite/Source/Composite/Public/Layers/CompositeLayerPlate.h
@@ -163,6 +163,12 @@ private:
 	/** Get number of valid passes. */
 	int32 GetValidPassesNum(TArrayView<const TObjectPtr<UCompositePassBase>> InPasses) const;
 
+	/** If Texture is a media profile media texture, attempt to open the associated media source */
+	void TryOpenMediaProfileSource();
+	
+	/** Attempt to close any media profile sources being played for this plate */
+	void TryCloseMediaProfileSource();
+	
 private:
 	/** Convenience function to find the last valid pass. */
 	int32 FindLastValidPassIndex(TArrayView<const TObjectPtr<UCompositePassBase>> InPasses) const;
diff --git a/Engine/Plugins/Experimental/Compositing/Composite/Source/CompositeEditor/CompositeEditor.Build.cs b/Engine/Plugins/Experimental/Compositing/Composite/Source/CompositeEditor/CompositeEditor.Build.cs
--- a/Engine/Plugins/Experimental/Compositing/Composite/Source/CompositeEditor/CompositeEditor.Build.cs
+++ b/Engine/Plugins/Experimental/Compositing/Composite/Source/CompositeEditor/CompositeEditor.Build.cs
@@ -26,6 +26,10 @@ public class CompositeEditor : ModuleRules
 				"Engine",
 				"InputCore",
 				"LevelEditor",
+				"MediaAssets",
+				"MediaFrameworkUtilities",
+				"Projects",
+				"PropertyEditor",
 				"UnrealEd",
 				"SceneOutliner",
                 "Slate",
diff --git a/Engine/Plugins/Experimental/Compositing/Composite/Source/CompositeEditor/Private/CompositeEditorStyle.cpp b/Engine/Plugins/Experimental/Compositing/Composite/Source/CompositeEditor/Private/CompositeEditorStyle.cpp
--- a/Engine/Plugins/Experimental/Compositing/Composite/Source/CompositeEditor/Private/CompositeEditorStyle.cpp
+++ b/Engine/Plugins/Experimental/Compositing/Composite/Source/CompositeEditor/Private/CompositeEditorStyle.cpp
@@ -2,6 +2,7 @@
 
 #include "CompositeEditorStyle.h"
 
+#include "Interfaces/IPluginManager.h"
 #include "Styling/SlateStyleMacros.h"
 #include "Styling/SlateStyleRegistry.h"
 
@@ -28,11 +29,15 @@ FCompositeEditorStyle::FCompositeEditorStyle()
 	const FVector2D Icon20x20(20.0f, 20.0f);
 	const FVector2D Icon64x64(64.0f, 64.0f);
 
-	FSlateStyleSet::SetContentRoot(FPaths::EngineContentDir() / TEXT("Editor/Slate/"));
+	const FString ContentDir = IPluginManager::Get().FindPlugin(TEXT("Composite"))->GetContentDir();
+	FSlateStyleSet::SetContentRoot(ContentDir);
+	FSlateStyleSet::SetCoreContentRoot(FPaths::EngineContentDir() / TEXT("Editor/Slate/"));
 	
-	Set("CompositeEditor.Composure", new IMAGE_BRUSH_SVG("Starship/Common/ComposureCompositing", Icon16x16));
+	Set("CompositeEditor.Composure", new CORE_IMAGE_BRUSH_SVG("Starship/Common/ComposureCompositing", Icon16x16));
 	
-    Set("CompositeEditor.Passes.Media", new IMAGE_BRUSH_SVG("Starship/AssetIcons/MediaPlayer_16", Icon16x16));
-	Set("CompositeEditor.Passes.Scene", new IMAGE_BRUSH_SVG("Starship/Common/Levels", Icon16x16));
-	Set("CompositeEditor.Passes.Layer", new IMAGE_BRUSH_SVG("Starship/Common/Layers", Icon16x16));
+    Set("CompositeEditor.Passes.Media", new CORE_IMAGE_BRUSH_SVG("Starship/AssetIcons/MediaPlayer_16", Icon16x16));
+	Set("CompositeEditor.Passes.Scene", new CORE_IMAGE_BRUSH_SVG("Starship/Common/Levels", Icon16x16));
+	Set("CompositeEditor.Passes.Layer", new CORE_IMAGE_BRUSH_SVG("Starship/Common/Layers", Icon16x16));
+
+	Set("CompositeEditor.MediaProfile", new IMAGE_BRUSH_SVG("Editor/Icons/MediaProfile", Icon16x16));
 }
diff --git a/Engine/Plugins/Experimental/Compositing/Composite/Source/CompositeEditor/Private/Customizations/CompositeLayerPlateCustomization.cpp b/Engine/Plugins/Experimental/Compositing/Composite/Source/CompositeEditor/Private/Customizations/CompositeLayerPlateCustomization.cpp
--- a/Engine/Plugins/Experimental/Compositing/Composite/Source/CompositeEditor/Private/Customizations/CompositeLayerPlateCustomization.cpp
+++ b/Engine/Plugins/Experimental/Compositing/Composite/Source/CompositeEditor/Private/Customizations/CompositeLayerPlateCustomization.cpp
@@ -9,20 +9,36 @@
 #include "DetailLayoutBuilder.h"
 #include "DetailWidgetRow.h"
 #include "Editor.h"
+#include "Engine/Texture2D.h"
 #include "IDetailGroup.h"
 #include "IPropertyUtilities.h"
+#include "MediaSource.h"
+#include "MediaTexture.h"
+#include "PropertyCustomizationHelpers.h"
 #include "ScopedTransaction.h"
 #include "UnrealEdGlobals.h"
 #include "Camera/CameraComponent.h"
 #include "Editor/UnrealEdEngine.h"
 #include "Framework/MultiBox/MultiBoxBuilder.h"
 #include "Layers/CompositeLayerPlate.h"
+#include "Profile/IMediaProfileManager.h"
+#include "Profile/MediaProfile.h"
+#include "Profile/MediaProfilePlaybackManager.h"
+#include "Styling/SlateIconFinder.h"
+#include "Subsystems/AssetEditorSubsystem.h"
 #include "UI/SCompositeActorPickerTable.h"
 #include "UI/SCompositePlatePassPanel.h"
+#include "Widgets/Images/SImage.h"
+#include "Widgets/Input/SComboButton.h"
 #include "Widgets/Text/STextBlock.h"
 
 #define LOCTEXT_NAMESPACE "FCompositeLayerPlateCustomization"
 
+namespace CompositeLayerPlateCustomization
+{
+	DECLARE_DELEGATE_TwoParams(FOnCustomizePropertyRow, const TSharedPtr<IPropertyHandle>&, IDetailPropertyRow&);
+}
+
 TSharedRef<IDetailCustomization> FCompositeLayerPlateCustomization::MakeInstance()
 {
 	return MakeShared<FCompositeLayerPlateCustomization>();
@@ -44,6 +60,14 @@ void FCompositeLayerPlateCustomization::CustomizeDetails(IDetailLayoutBuilder& D
 	
 	IDetailCategoryBuilder& PlateCategory = DetailLayout.EditCategoryAllowNone(NAME_None);
 
+	TMap<FName, CompositeLayerPlateCustomization::FOnCustomizePropertyRow> PropertyCustomizations =
+	{
+		{
+			GET_MEMBER_NAME_CHECKED(UCompositeLayerPlate, Texture),
+			CompositeLayerPlateCustomization::FOnCustomizePropertyRow::CreateSP(this, &FCompositeLayerPlateCustomization::CustomizeTexturePropertyRow)
+		}
+	};
+	
 	// Add all default simple properties ahead of the custom group
 	TArray<TSharedRef<IPropertyHandle>> PlateProperties;
 	PlateCategory.GetDefaultProperties(PlateProperties, /* bSimpleProperties */ true, /* bAdvancedProperties */ false);
@@ -51,7 +75,12 @@ void FCompositeLayerPlateCustomization::CustomizeDetails(IDetailLayoutBuilder& D
 	{
 		if (!Property->IsCustomized())
 		{
-			PlateCategory.AddProperty(Property);
+			const FName PropertyName = Property->GetProperty()->GetFName();
+			IDetailPropertyRow& PropertyRow = PlateCategory.AddProperty(Property);
+			if (PropertyCustomizations.Contains(PropertyName))
+			{
+				PropertyCustomizations[PropertyName].ExecuteIfBound(Property, PropertyRow);
+			}
 		}
 	}
 	
@@ -98,6 +127,144 @@ void FCompositeLayerPlateCustomization::CustomizeDetails(IDetailLayoutBuilder& D
 	}
 }
 
+void FCompositeLayerPlateCustomization::CustomizeTexturePropertyRow(const TSharedPtr<IPropertyHandle>& InPropertyHandle, IDetailPropertyRow& InPropertyRow)
+{
+	TexturePropertyHandle = InPropertyHandle;
+
+	InPropertyRow.CustomWidget()
+	.NameContent()
+	[
+		InPropertyHandle->CreatePropertyNameWidget()
+	]
+	.ValueContent()
+	[
+		SNew(SObjectPropertyEntryBox)
+		.PropertyHandle(InPropertyHandle)
+		.AllowCreate(true)
+		.ThumbnailPool(CachedDetailBuilder.IsValid() ? CachedDetailBuilder.Pin()->GetThumbnailPool() : nullptr)
+		.OnShouldFilterAsset_Lambda([](FAssetData const& AssetData)
+		{
+			if (UClass* AssetClass = AssetData.GetClass())
+			{
+				return !(AssetClass->IsChildOf(UTexture2D::StaticClass()) || AssetClass->IsChildOf(UMediaTexture::StaticClass()));
+			}
+
+			return true;
+		})
+		.CustomContentSlot()
+		[
+			SNew(SHorizontalBox)
+
+			+SHorizontalBox::Slot()
+			.AutoWidth()
+			[
+				SNew(SComboButton)
+				.HasDownArrow(true)
+				.ComboButtonStyle(FAppStyle::Get(), "SimpleComboButton")
+				.ToolTipText(LOCTEXT("MediaProfileSourceSelectorToolTip", "Select media texture from a media source in the active Media Profile"))
+				.HAlign(HAlign_Center)
+				.VAlign(VAlign_Center)
+				.OnGetMenuContent(this, &FCompositeLayerPlateCustomization::GetMediaProfileSourceSelectorMenu)
+				.IsEnabled(this, &FCompositeLayerPlateCustomization::HasActiveMediaProfile)
+				.ButtonContent()
+				[
+					SNew(SImage)
+					.Image(FCompositeEditorStyle::Get().GetBrush("CompositeEditor.MediaProfile"))
+				]
+			]
+		]
+	];
+}
+
+bool FCompositeLayerPlateCustomization::HasActiveMediaProfile() const
+{
+	UMediaProfile* ActiveMediaProfile = IMediaProfileManager::Get().GetCurrentMediaProfile();
+	return IsValid(ActiveMediaProfile);
+}
+
+TSharedRef<SWidget> FCompositeLayerPlateCustomization::GetMediaProfileSourceSelectorMenu()
+{
+	UMediaProfile* ActiveMediaProfile = IMediaProfileManager::Get().GetCurrentMediaProfile();
+	if (!ActiveMediaProfile)
+	{
+		return SNullWidget::NullWidget;
+	}
+	
+	const bool bCloseAfterSelection = true;
+	FMenuBuilder MenuBuilder(bCloseAfterSelection, nullptr);
+
+	MenuBuilder.BeginSection(NAME_None, LOCTEXT("MediaProfileSourcesSectionLabel", "Texture from Media Profile"));
+	{
+		bool bNoMediaSources = true;
+		for (int32 Index = 0; Index < ActiveMediaProfile->NumMediaSources(); ++Index)
+		{
+			if (UMediaSource* MediaSource = ActiveMediaProfile->GetMediaSource(Index))
+			{
+				bNoMediaSources = false;
+
+				MenuBuilder.AddMenuEntry(
+					FText::FromString(ActiveMediaProfile->GetLabelForMediaSource(Index)),
+					FText::GetEmpty(),
+					FSlateIconFinder::FindIconForClass(MediaSource->GetClass()),
+					FUIAction(FExecuteAction::CreateSP(this, &FCompositeLayerPlateCustomization::SetTextureToMediaProfileSource, Index))
+				);
+			}
+		}
+
+		if (bNoMediaSources)
+		{
+			TSharedRef<SWidget> TextBlock = SNew(STextBlock)
+				.Text(LOCTEXT("NoMediaSources", "No Media Sources"))
+				.Font(FCoreStyle::GetDefaultFontStyle("Italic", 10));
+			
+			MenuBuilder.AddWidget(TextBlock,FText::GetEmpty());
+		}
+	}
+	MenuBuilder.EndSection();
+
+	MenuBuilder.AddSeparator();
+	MenuBuilder.BeginSection(NAME_None);
+	{
+		MenuBuilder.AddMenuEntry(
+			FText::Format(LOCTEXT("OpenActiveMediaProfileLabelFormat", "Open {0}"), FText::FromString(ActiveMediaProfile->GetName())),
+			FText::GetEmpty(),
+			FSlateIconFinder::FindIconForClass(ActiveMediaProfile->StaticClass()),
+			FUIAction(FExecuteAction::CreateSP(this, &FCompositeLayerPlateCustomization::OpenMediaProfile))
+		);
+	}
+	MenuBuilder.EndSection();
+	
+	return MenuBuilder.MakeWidget();
+}
+
+void FCompositeLayerPlateCustomization::SetTextureToMediaProfileSource(int32 InMediaSourceIndex)
+{
+	UMediaProfile* ActiveMediaProfile = IMediaProfileManager::Get().GetCurrentMediaProfile();
+	if (!ActiveMediaProfile)
+	{
+		return;
+	}
+
+	if (!IsValid(ActiveMediaProfile->GetMediaSource(InMediaSourceIndex)))
+	{
+		return;
+	}
+
+	UMediaTexture* MediaTexture = ActiveMediaProfile->GetPlaybackManager()->GetSourceMediaTextureFromIndex(InMediaSourceIndex);
+	TexturePropertyHandle->SetValue(MediaTexture);
+}
+
+void FCompositeLayerPlateCustomization::OpenMediaProfile()
+{
+	UMediaProfile* ActiveMediaProfile = IMediaProfileManager::Get().GetCurrentMediaProfile();
+	if (!ActiveMediaProfile)
+	{
+		return;
+	}
+	
+	GEditor->GetEditorSubsystem<UAssetEditorSubsystem>()->OpenEditorForAsset(ActiveMediaProfile);
+}
+
 void FCompositeLayerPlateCustomization::OnExtendCompositeMeshAddMenu(FMenuBuilder& MenuBuilder)
 {
 	MenuBuilder.AddMenuEntry(
diff --git a/Engine/Plugins/Experimental/Compositing/Composite/Source/CompositeEditor/Private/Customizations/CompositeLayerPlateCustomization.h b/Engine/Plugins/Experimental/Compositing/Composite/Source/CompositeEditor/Private/Customizations/CompositeLayerPlateCustomization.h
--- a/Engine/Plugins/Experimental/Compositing/Composite/Source/CompositeEditor/Private/Customizations/CompositeLayerPlateCustomization.h
+++ b/Engine/Plugins/Experimental/Compositing/Composite/Source/CompositeEditor/Private/Customizations/CompositeLayerPlateCustomization.h
@@ -4,6 +4,9 @@
 
 #include "IDetailCustomization.h"
 
+class SWidget;
+class IDetailPropertyRow;
+class IPropertyHandle;
 class SCompositePlatePassPanel;
 class FMenuBuilder;
 
@@ -22,6 +25,21 @@ public:
 	// End of IDetailCustomization interface
 
 private:
+	/** Customizes the Texture property in UCompositeLayerPlate to add in the media profile source selector */
+	void CustomizeTexturePropertyRow(const TSharedPtr<IPropertyHandle>& InPropertyHandle, IDetailPropertyRow& InPropertyRow);
+
+	/** Gets whether there is an active media profile configured in the editor */
+	bool HasActiveMediaProfile() const;
+
+	/** Gets the dropdown menu widget to display when the media profile button is pressed */
+	TSharedRef<SWidget> GetMediaProfileSourceSelectorMenu();
+
+	/** Sets the Texture property handle's value to the media texture corresponding to the specified media source in the active media profile */
+	void SetTextureToMediaProfileSource(int32 InMediaSourceIndex);
+
+	/** Opens the active media profile in the media profile editor */
+	void OpenMediaProfile();
+	
 	/** Raised by the composite mesh actor list when building the Add Actor menu */
 	void OnExtendCompositeMeshAddMenu(FMenuBuilder& MenuBuilder);
 
@@ -36,4 +54,5 @@ private:
 	
 private:
 	TWeakPtr<IDetailLayoutBuilder> CachedDetailBuilder;
+	TSharedPtr<IPropertyHandle> TexturePropertyHandle;
 };
diff --git a/Engine/Plugins/Media/MediaFrameworkUtilities/Source/MediaFrameworkUtilities/Private/Profile/MediaProfilePlaybackManager.cpp b/Engine/Plugins/Media/MediaFrameworkUtilities/Source/MediaFrameworkUtilities/Private/Profile/MediaProfilePlaybackManager.cpp
--- a/Engine/Plugins/Media/MediaFrameworkUtilities/Source/MediaFrameworkUtilities/Private/Profile/MediaProfilePlaybackManager.cpp
+++ b/Engine/Plugins/Media/MediaFrameworkUtilities/Source/MediaFrameworkUtilities/Private/Profile/MediaProfilePlaybackManager.cpp
@@ -51,7 +51,7 @@ UMediaTexture* UMediaProfilePlaybackManager::GetSourceMediaTextureFromIndex(int3
 		return nullptr;	
 	}
 
-	return MediaSourceTextures[InMediaSourceIndex].LoadSynchronous();
+	return GIsRunning ? MediaSourceTextures[InMediaSourceIndex].LoadSynchronous() : nullptr;
 }
 
 UMediaTexture* UMediaProfilePlaybackManager::GetSourceMediaTexture(UMediaSource* InMediaSource) const
@@ -66,6 +66,29 @@ UMediaTexture* UMediaProfilePlaybackManager::GetSourceMediaTexture(UMediaSource*
 	return GetSourceMediaTextureFromIndex(MediaSourceIndex);
 }
 
+bool UMediaProfilePlaybackManager::IsValidSourceMediaTexture(UMediaTexture* InMediaTexture, int32& OutMediaSourceIndex) const
+{
+	OutMediaSourceIndex = INDEX_NONE;
+	
+	UMediaProfile* MediaProfile = GetOwningMediaProfile();
+	if (!MediaProfile)
+	{
+		return false;
+	}
+	
+	int32 MediaSourceIndex;
+	if (MediaSourceTextures.Find(InMediaTexture, MediaSourceIndex))
+	{
+		if (MediaProfile->NumMediaSources() > MediaSourceIndex)
+		{
+			OutMediaSourceIndex = MediaSourceIndex;
+			return true;
+		}
+	}
+
+	return false;
+}
+
 UMediaTexture* UMediaProfilePlaybackManager::OpenSourceFromIndex(int32 InMediaSourceIndex, void* InConsumer)
 {
 	if (!MediaSourceTextures.IsValidIndex(InMediaSourceIndex))
@@ -170,6 +193,26 @@ void UMediaProfilePlaybackManager::CloseSource(UMediaSource* InMediaSource, cons
 	CloseSourceFromIndex(MediaSourceIndex, InArgs);
 }
 
+void UMediaProfilePlaybackManager::CloseSourcesForConsumer(const FCloseSourceArgs& InArgs)
+{
+	TArray<TObjectPtr<UMediaPlayer>> MediaPlayersUsedByConsumer;
+	for (const TPair<TObjectPtr<UMediaPlayer>, TSet<void*>>& Pair : SourceConsumers)
+	{
+		if (Pair.Value.Contains(InArgs.Consumer))
+		{
+			MediaPlayersUsedByConsumer.Add(Pair.Key);
+		}
+	}
+
+	for (const TObjectPtr<UMediaPlayer>& MediaPlayer : MediaPlayersUsedByConsumer)
+	{
+		if (const int32* MediaSourceIndex = MediaSourcePlayers.FindKey(MediaPlayer))
+		{
+			CloseSourceFromIndex(*MediaSourceIndex, InArgs);
+		}
+	}
+}
+
 UMediaPlayer* UMediaProfilePlaybackManager::FindOrCreateMediaPlayerForIndex(int32 InMediaSourceIndex)
 {
 	UMediaProfile* MediaProfile = GetOwningMediaProfile();
diff --git a/Engine/Plugins/Media/MediaFrameworkUtilities/Source/MediaFrameworkUtilities/Public/Profile/MediaProfilePlaybackManager.h b/Engine/Plugins/Media/MediaFrameworkUtilities/Source/MediaFrameworkUtilities/Public/Profile/MediaProfilePlaybackManager.h
--- a/Engine/Plugins/Media/MediaFrameworkUtilities/Source/MediaFrameworkUtilities/Public/Profile/MediaProfilePlaybackManager.h
+++ b/Engine/Plugins/Media/MediaFrameworkUtilities/Source/MediaFrameworkUtilities/Public/Profile/MediaProfilePlaybackManager.h
@@ -35,6 +35,14 @@ public:
 	 */
 	UMediaTexture* GetSourceMediaTexture(UMediaSource* InMediaSource) const;
 
+	/**
+	 * Checks the specified media texture to see if it is being used by the media profile to play one of its media sources
+	 * @param InMediaTexture The media texture to test
+	 * @param OutMediaSourceIndex The media source index in the media profile that is using the texture, or INDEX_NONE if there is none
+	 * @return true if the media texture is being used by the active media profile; otherwise false
+	 */
+	bool IsValidSourceMediaTexture(UMediaTexture* InMediaTexture, int32& OutMediaSourceIndex) const;
+	
 	/**
 	 * Opens and begins playing the specified media source 
 	 * @param InMediaSourceIndex The index of the media source within the media profile
@@ -83,6 +91,12 @@ public:
 	 */
 	void CloseSource(UMediaSource* InMediaSource, const FCloseSourceArgs& InArgs = FCloseSourceArgs());
 
+	/**
+	 * Closes all media sources that are being played by a specified consumer
+	 * @param InArgs Arguments that control what happens when the media source is closed
+	 */
+	void CloseSourcesForConsumer(const FCloseSourceArgs& InArgs);
+	
 private:
 	/**
 	 * Attempts to find an existing media player for the specified source, and creates one if an existing one is not found
