Composite: Added new non-transactional active state, to ensure only one actor is rendering at a time.

* Active state is automatically disabled when joining a multi-user session.
* Active state is automatically disabled on new actors when more than one is added to a level.

diff --git a/Engine/Plugins/Experimental/Compositing/Composite/Source/Composite/Private/CompositeActor.cpp b/Engine/Plugins/Experimental/Compositing/Composite/Source/Composite/Private/CompositeActor.cpp
--- a/Engine/Plugins/Experimental/Compositing/Composite/Source/Composite/Private/CompositeActor.cpp
+++ b/Engine/Plugins/Experimental/Compositing/Composite/Source/Composite/Private/CompositeActor.cpp
@@ -18,6 +18,7 @@
 #include "Components/PrimitiveComponent.h"
 #include "Camera/CameraActor.h"
 #include "Camera/CameraComponent.h"
+#include "EngineUtils.h"
 #include "Engine/Engine.h"
 #include "Engine/World.h"
 #include "Engine/Texture.h"
@@ -200,6 +201,7 @@ namespace UE
 ACompositeActor::ACompositeActor(const FObjectInitializer& ObjectInitializer)
 	: Super(ObjectInitializer)
 	, ViewProjectionComponent{}
+	, bIsActive{true}
 	, bIsEnabled{true}
 // TODO: Can we find the current viewport render resolution instead?
 // View rects are only accessible from the render thread, but maybe there is an alternative.
@@ -275,7 +277,7 @@ void ACompositeActor::PostRegisterAllComponents()
 	}
 #endif
 
-	if (IsActive())
+	if (IsRendering())
 	{
 		ApplyAllState();
 	}
@@ -306,7 +308,7 @@ void ACompositeActor::Tick(float DeltaTime)
 
 	SCOPE_CYCLE_COUNTER(STAT_Composite_Actor_Tick);
 
-	if (!IsActive())
+	if (!IsRendering())
 	{
 		return;
 	}
@@ -381,7 +383,11 @@ void ACompositeActor::PostEditChangeChainProperty(struct FPropertyChangedChainEv
 
 	const FName PropertyName = PropertyChangedChainEvent.GetPropertyName();
 
-	if (PropertyName == GET_MEMBER_NAME_CHECKED(ThisClass, bIsEnabled))
+	if (PropertyName == GET_MEMBER_NAME_CHECKED(ThisClass, bIsActive))
+	{
+		SetActive(bIsActive);
+	}
+	else if (PropertyName == GET_MEMBER_NAME_CHECKED(ThisClass, bIsEnabled))
 	{
 		SetEnabled(bIsEnabled);
 	}
@@ -477,7 +483,7 @@ void ACompositeActor::SetIsTemporarilyHiddenInEditor(bool bIsHidden)
 {
 	Super::SetIsTemporarilyHiddenInEditor(bIsHidden);
 
-	if (IsActive())
+	if (IsRendering())
 	{
 		ApplyAllState();
 	}
@@ -493,6 +499,39 @@ bool ACompositeActor::ShouldTickIfViewportsOnly() const
 }
 #endif
 
+bool ACompositeActor::IsActive() const
+{
+	return bIsActive;
+}
+
+void ACompositeActor::SetActive(bool bInActive)
+{
+	bIsActive = bInActive;
+
+	if (IsRendering())
+	{
+		ApplyAllState();
+	}
+	else
+	{
+		RemoveAllState();
+	}
+
+	const UWorld* World = GetWorld();
+
+	if (bIsActive && IsValid(World))
+	{
+		for (TActorIterator<ACompositeActor> It(World); It; ++It)
+		{
+			if (*It != this)
+			{
+				It->bIsActive = false;
+				It->RemoveAllState();
+			}
+		}
+	}
+}
+
 bool ACompositeActor::IsEnabled() const
 {
 	return bIsEnabled;
@@ -502,7 +541,7 @@ void ACompositeActor::SetEnabled(bool bInEnabled)
 {
 	bIsEnabled = bInEnabled;
 
-	if (IsActive())
+	if (IsRendering())
 	{
 		ApplyAllState();
 	}
@@ -667,18 +706,25 @@ void ACompositeActor::DestroySceneCaptures(const UCompositeLayerBase* InLayer)
 #endif
 }
 
-bool ACompositeActor::IsActive() const
+bool ACompositeActor::IsRendering() const
 {
-	bool bIsActive = bIsEnabled;
+	bool bIsRendering = IsEnabled() && IsActive();
 
 #if WITH_EDITOR
 	if (IsTemporarilyHiddenInEditor())
 	{
-		bIsActive = false;
+		bIsRendering = false;
 	}
 #endif
 
-	return bIsActive;
+	return bIsRendering;
+}
+
+void ACompositeActor::JoinConcertSession()
+{
+	SetActive(false);
+
+	OnJoinMultiUserSession.Broadcast();
 }
 
 void ACompositeActor::ApplyAllState()
diff --git a/Engine/Plugins/Experimental/Compositing/Composite/Source/Composite/Private/Layers/CompositeLayerPlate.cpp b/Engine/Plugins/Experimental/Compositing/Composite/Source/Composite/Private/Layers/CompositeLayerPlate.cpp
--- a/Engine/Plugins/Experimental/Compositing/Composite/Source/Composite/Private/Layers/CompositeLayerPlate.cpp
+++ b/Engine/Plugins/Experimental/Compositing/Composite/Source/Composite/Private/Layers/CompositeLayerPlate.cpp
@@ -113,7 +113,7 @@ void UCompositeLayerPlate::Tick(float DeltaTime)
 
 	ACompositeActor* CompositeActor = GetTypedOuter<ACompositeActor>();
 	
-	if (!IsValid(CompositeActor) || !CompositeActor->IsActive())
+	if (!IsValid(CompositeActor) || !CompositeActor->IsRendering())
 	{
 		return;
 	}
diff --git a/Engine/Plugins/Experimental/Compositing/Composite/Source/Composite/Public/CompositeActor.h b/Engine/Plugins/Experimental/Compositing/Composite/Source/Composite/Public/CompositeActor.h
--- a/Engine/Plugins/Experimental/Compositing/Composite/Source/Composite/Public/CompositeActor.h
+++ b/Engine/Plugins/Experimental/Compositing/Composite/Source/Composite/Public/CompositeActor.h
@@ -18,6 +18,7 @@ class UCompositeShadowReflectionCatcherComponent;
 class UPrimitiveComponent;
 class USceneCaptureComponent2D;
 
+DECLARE_DYNAMIC_MULTICAST_DELEGATE(FOnCompositeActorJoinMultiUserSession);
 DECLARE_DELEGATE_TwoParams(FOnCompositeActorEnabledStateChange, bool /* */, UCompositeCoreSubsystem* /* */);
 
 USTRUCT()
@@ -71,6 +72,14 @@ public:
 	UE_API virtual bool ShouldTickIfViewportsOnly() const override;
 #endif // WITH_EDITOR
 
+	/** Get the (locally) active state. */
+	UFUNCTION(BlueprintGetter)
+	UE_API bool IsActive() const;
+
+	/** Set the (locally) active state. */
+	UFUNCTION(BlueprintSetter)
+	UE_API void SetActive(bool bInActive);
+
 	/** Get the enabled state. */
 	UFUNCTION(BlueprintGetter)
 	UE_API bool IsEnabled() const;
@@ -95,8 +104,12 @@ private:
 	UPROPERTY(Instanced)
 	TObjectPtr<class UCompositeViewProjectionComponent> ViewProjectionComponent;
 
-	/** Whether or not the composite behavior is active. */
-	UPROPERTY(EditAnywhere, BlueprintGetter = IsEnabled, BlueprintSetter = SetEnabled, Interp, Category = "Composite", meta = (AllowPrivateAccess = true, DisplayPriority="0"))
+	/** Whether or not the composite behavior is active locally - used primarily for multi-user. */
+	UPROPERTY(NonTransactional, EditInstanceOnly, BlueprintGetter = IsActive, BlueprintSetter = SetActive, Category = "Composite", meta = (AllowPrivateAccess = true))
+	bool bIsActive;
+
+	/** Whether or not the composite behavior is enabled. */
+	UPROPERTY(EditAnywhere, BlueprintGetter = IsEnabled, BlueprintSetter = SetEnabled, Interp, Category = "Composite", meta = (DisplayName = "Enabled", AllowPrivateAccess = true, DisplayPriority = "0"))
 	bool bIsEnabled;
 
 public:
@@ -130,6 +143,10 @@ public:
 	UPROPERTY(EditAnywhere, BlueprintReadWrite, Category = "Composite", AdvancedDisplay)
 	ECompositeMainRenderOutputMode MainRenderOutput;
 
+	/** Delegate called when joining multi-user. Can be used to automatically re-active the actor locally. */
+	UPROPERTY(BlueprintAssignable, BlueprintCallable, BlueprintReadWrite, Category = "Composite")
+	FOnCompositeActorJoinMultiUserSession OnJoinMultiUserSession;
+
 public:
 	/** Convenience function used by layers to request a scene capture component. */
 	template<class RetType>
@@ -141,8 +158,11 @@ public:
 	/** Convenience function used by layers to remove and scene capture components. */
 	UE_API void DestroySceneCaptures(const UCompositeLayerBase* InLayer);
 
-	/** Returns true if the composite is active. */
-	UE_API bool IsActive() const;
+	/** Returns true if the composite is actively rendering. */
+	UE_API bool IsRendering() const;
+
+	/** Called when we join a multi-user session. */
+	UE_API void JoinConcertSession();
 
 private:
 	/** Apply all actor state to the composite subsystem. */
@@ -162,6 +182,10 @@ private:
 
 	/** Flag used to determine in which cases we should ignore PostRegister / Unregister calls. */
 	bool bIsModifyingAProperty;
+
+	/** Befriend the actor's editor customizations. */
+	friend class FCompositeActorCustomization;
+	friend class FCompositeActorPanelDetailCustomization;
 };
 
 #undef UE_API
diff --git a/Engine/Plugins/Experimental/Compositing/Composite/Source/CompositeEditor/CompositeEditor.Build.cs b/Engine/Plugins/Experimental/Compositing/Composite/Source/CompositeEditor/CompositeEditor.Build.cs
--- a/Engine/Plugins/Experimental/Compositing/Composite/Source/CompositeEditor/CompositeEditor.Build.cs
+++ b/Engine/Plugins/Experimental/Compositing/Composite/Source/CompositeEditor/CompositeEditor.Build.cs
@@ -21,6 +21,7 @@ public class CompositeEditor : ModuleRules
 			{
 				"ApplicationCore",
 				"CompositeCore",
+				"ConcertSyncClient",
 				"CoreUObject",
 				"EditorWidgets",
 				"Engine",
diff --git a/Engine/Plugins/Experimental/Compositing/Composite/Source/CompositeEditor/Private/CompositeEditorModule.cpp b/Engine/Plugins/Experimental/Compositing/Composite/Source/CompositeEditor/Private/CompositeEditorModule.cpp
--- a/Engine/Plugins/Experimental/Compositing/Composite/Source/CompositeEditor/Private/CompositeEditorModule.cpp
+++ b/Engine/Plugins/Experimental/Compositing/Composite/Source/CompositeEditor/Private/CompositeEditorModule.cpp
@@ -6,6 +6,8 @@
 #include "EngineUtils.h"
 #include "EngineAnalytics.h"
 #include "Framework/Notifications/NotificationManager.h"
+#include "IConcertSyncClient.h"
+#include "IConcertSyncClientModule.h"
 #include "LevelEditor.h"
 #include "Modules/ModuleManager.h"
 #include "Widgets/Notifications/SNotificationList.h"
@@ -13,6 +15,7 @@
 #include "CompositeEditorCommands.h"
 #include "CompositeEditorStyle.h"
 #include "PropertyEditorModule.h"
+#include "Customizations/CompositeActorCustomization.h"
 #include "Customizations/CompositeLayerPlateCustomization.h"
 #include "Customizations/CompositeLayerShadowReflectionCustomization.h"
 #include "Layers/CompositeLayerPlate.h"
@@ -36,6 +39,11 @@ void FCompositeEditorModule::StartupModule()
 
 void FCompositeEditorModule::ShutdownModule()
 {
+	if (ConcertSyncClient.IsValid())
+	{
+		ConcertSyncClient->GetConcertClient()->OnSessionStartup().RemoveAll(this);
+	}
+
 	UnregisterCustomizations();
 	UnregisterTabSpawners();
 
@@ -46,6 +54,8 @@ void FCompositeEditorModule::ShutdownModule()
 	{
 		GEditor->OnLevelActorAdded().RemoveAll(this);
 	}
+	
+	FCoreDelegates::OnPostEngineInit.RemoveAll(this);
 }
 
 void FCompositeEditorModule::OnPostEngineInit()
@@ -54,6 +64,12 @@ void FCompositeEditorModule::OnPostEngineInit()
 	{
 		GEditor->OnLevelActorAdded().AddRaw(this, &FCompositeEditorModule::OnLevelActorAdded);
 	}
+
+	ConcertSyncClient = IConcertSyncClientModule::Get().GetClient(TEXT("MultiUser"));
+	if (ConcertSyncClient.IsValid())
+	{
+		ConcertSyncClient->GetConcertClient()->OnSessionStartup().AddRaw(this, &FCompositeEditorModule::OnConcertSessionStartup);
+	}
 }
 
 void FCompositeEditorModule::OnLevelActorAdded(AActor* InActor)
@@ -79,9 +95,9 @@ void FCompositeEditorModule::OnLevelActorAdded(AActor* InActor)
 			ACompositeActor* CompositeActor = Cast<ACompositeActor>(InActor);
 			if (IsValid(CompositeActor))
 			{
-				CompositeActor->SetEnabled(false);
+				CompositeActor->SetActive(false);
 
-				UE_LOG(LogCompositeEditor, Warning, TEXT("No more than one composite actor should be enabled at the same time, disabling."));
+				UE_LOG(LogCompositeEditor, Display, TEXT("No more than one composite actor should be active at the same time."));
 			}
 			break;
 		}
@@ -93,6 +109,23 @@ void FCompositeEditorModule::OnLevelActorAdded(AActor* InActor)
 	}
 }
 
+void FCompositeEditorModule::OnConcertSessionStartup(TSharedRef<IConcertClientSession> InSession)
+{
+	const UWorld* World = GEditor->GetEditorWorldContext().World();
+
+	if (IsValid(World))
+	{
+		for (TActorIterator<ACompositeActor> It(World); It; ++It)
+		{
+			ACompositeActor* CompositeActor = *It;
+			if (IsValid(CompositeActor))
+			{
+				CompositeActor->JoinConcertSession();
+			}
+		}
+	}
+}
+
 void FCompositeEditorModule::RegisterTabSpawners()
 {
 	FLevelEditorModule& LevelEditorModule = FModuleManager::GetModuleChecked<FLevelEditorModule>("LevelEditor");
@@ -123,6 +156,11 @@ void FCompositeEditorModule::RegisterCustomizations()
 		UCompositeLayerShadowReflection::StaticClass()->GetFName(),
 		FOnGetDetailCustomizationInstance::CreateStatic(&FCompositeLayerShadowReflectionCustomization::MakeInstance)
 	);
+
+	PropertyModule.RegisterCustomClassLayout(
+		ACompositeActor::StaticClass()->GetFName(),
+		FOnGetDetailCustomizationInstance::CreateStatic(&FCompositeActorCustomization::MakeInstance)
+	);
 }
 
 void FCompositeEditorModule::UnregisterCustomizations()
@@ -131,6 +169,7 @@ void FCompositeEditorModule::UnregisterCustomizations()
 	{
 		if (FPropertyEditorModule* PropertyModule = FModuleManager::Get().GetModulePtr<FPropertyEditorModule>("PropertyEditor"))
 		{
+			PropertyModule->UnregisterCustomClassLayout(ACompositeActor::StaticClass()->GetFName());
 			PropertyModule->UnregisterCustomClassLayout(UCompositeLayerPlate::StaticClass()->GetFName());
 			PropertyModule->UnregisterCustomClassLayout(UCompositeLayerShadowReflection::StaticClass()->GetFName());
 		}
diff --git a/Engine/Plugins/Experimental/Compositing/Composite/Source/CompositeEditor/Private/Customizations/CompositeActorCustomization.cpp b/Engine/Plugins/Experimental/Compositing/Composite/Source/CompositeEditor/Private/Customizations/CompositeActorCustomization.cpp
new file mode 100644
--- /dev/null
+++ b/Engine/Plugins/Experimental/Compositing/Composite/Source/CompositeEditor/Private/Customizations/CompositeActorCustomization.cpp
@@ -0,0 +1,20 @@
+﻿// Copyright Epic Games, Inc. All Rights Reserved.
+
+#include "CompositeActorCustomization.h"
+
+#include "CompositeActor.h"
+#include "DetailLayoutBuilder.h"
+
+#define LOCTEXT_NAMESPACE "FCompositeActorCustomization"
+
+TSharedRef<IDetailCustomization> FCompositeActorCustomization::MakeInstance()
+{
+	return MakeShared<FCompositeActorCustomization>();
+}
+
+void FCompositeActorCustomization::CustomizeDetails(IDetailLayoutBuilder& DetailLayout)
+{
+	DetailLayout.HideProperty(GET_MEMBER_NAME_CHECKED(ACompositeActor, bIsActive));
+}
+
+#undef LOCTEXT_NAMESPACE
diff --git a/Engine/Plugins/Experimental/Compositing/Composite/Source/CompositeEditor/Private/Customizations/CompositeActorCustomization.h b/Engine/Plugins/Experimental/Compositing/Composite/Source/CompositeEditor/Private/Customizations/CompositeActorCustomization.h
new file mode 100644
--- /dev/null
+++ b/Engine/Plugins/Experimental/Compositing/Composite/Source/CompositeEditor/Private/Customizations/CompositeActorCustomization.h
@@ -0,0 +1,19 @@
+﻿// Copyright Epic Games, Inc. All Rights Reserved.
+
+#pragma once
+
+#include "IDetailCustomization.h"
+
+/**
+ * Customization for the composite actor.
+ */
+class FCompositeActorCustomization : public IDetailCustomization
+{
+public:
+	/** Makes a new instance of the details customization */
+	static TSharedRef<IDetailCustomization> MakeInstance();
+
+	// IDetailCustomization interface
+	virtual void CustomizeDetails(IDetailLayoutBuilder& DetailLayout) override;
+	// End of IDetailCustomization interface
+};
diff --git a/Engine/Plugins/Experimental/Compositing/Composite/Source/CompositeEditor/Private/UI/SCompositeEditorPanel.cpp b/Engine/Plugins/Experimental/Compositing/Composite/Source/CompositeEditor/Private/UI/SCompositeEditorPanel.cpp
--- a/Engine/Plugins/Experimental/Compositing/Composite/Source/CompositeEditor/Private/UI/SCompositeEditorPanel.cpp
+++ b/Engine/Plugins/Experimental/Compositing/Composite/Source/CompositeEditor/Private/UI/SCompositeEditorPanel.cpp
@@ -377,6 +377,7 @@ public:
 				})
 			];
 		
+		DetailBuilder.HideProperty(GET_MEMBER_NAME_CHECKED(ACompositeActor, bIsActive));
 		DetailBuilder.HideProperty(GET_MEMBER_NAME_CHECKED(ACompositeActor, CompositeLayers));
 	}
 
diff --git a/Engine/Plugins/Experimental/Compositing/Composite/Source/CompositeEditor/Private/UI/SCompositePanelLayerTree.cpp b/Engine/Plugins/Experimental/Compositing/Composite/Source/CompositeEditor/Private/UI/SCompositePanelLayerTree.cpp
--- a/Engine/Plugins/Experimental/Compositing/Composite/Source/CompositeEditor/Private/UI/SCompositePanelLayerTree.cpp
+++ b/Engine/Plugins/Experimental/Compositing/Composite/Source/CompositeEditor/Private/UI/SCompositePanelLayerTree.cpp
@@ -28,6 +28,7 @@
 
 namespace CompositePanelLayerTree
 {
+	static const FName LayerTreeColumn_Active = "Active";
 	static const FName LayerTreeColumn_Enabled = "Enabled";
 	static const FName LayerTreeColumn_Solo = "Solo";
 	static const FName LayerTreeColumn_Name = "Name";
@@ -346,19 +347,41 @@ public:
 private:
 	virtual TSharedRef<SWidget> GenerateWidgetForColumn(const FName& InColumnName) override
 	{
-		if (InColumnName == CompositePanelLayerTree::LayerTreeColumn_Enabled)
+		if (InColumnName == CompositePanelLayerTree::LayerTreeColumn_Active)
 		{
-			return SNew(SCheckBox)
-				.IsChecked(this, &SCompositePanelLayerTreeItemRow::IsItemEnabled)
-				.OnCheckStateChanged(this, &SCompositePanelLayerTreeItemRow::OnIsItemEnabledChanged);
+			return SNew(SHorizontalBox)
+
+				+ SHorizontalBox::Slot()
+				.AutoWidth()
+				.Padding(4.f, 0.f, 0.f, 0.f)
+				[
+					SNew(SCheckBox)
+						.Style(FAppStyle::Get(), "RadioButton")
+						.Visibility(this, &SCompositePanelLayerTreeItemRow::IsItemActiveVisible)
+						.IsChecked(this, &SCompositePanelLayerTreeItemRow::IsItemActive)
+						.OnCheckStateChanged(this, &SCompositePanelLayerTreeItemRow::OnIsItemActiveChanged)
+				];
+		}
+		else if (InColumnName == CompositePanelLayerTree::LayerTreeColumn_Enabled)
+		{
+			return SNew(SHorizontalBox)
+
+				+ SHorizontalBox::Slot()
+				.AutoWidth()
+				.Padding(4.f, 0.f, 0.f, 0.f)
+				[
+					SNew(SCheckBox)
+						.IsChecked(this, &SCompositePanelLayerTreeItemRow::IsItemEnabled)
+						.OnCheckStateChanged(this, &SCompositePanelLayerTreeItemRow::OnIsItemEnabledChanged)
+				];
 		}
-		if (InColumnName == CompositePanelLayerTree::LayerTreeColumn_Solo)
+		else if (InColumnName == CompositePanelLayerTree::LayerTreeColumn_Solo)
 		{
 			return SNew(SImage).Image(FAppStyle::GetBrush("MediaAsset.AssetActions.Solo.Small"))
 				.ColorAndOpacity_Raw(this, &SCompositePanelLayerTreeItemRow::GetItemSoloIconColor)
 				.OnMouseButtonDown(this, &SCompositePanelLayerTreeItemRow::OnItemSoloIconMouseButtonDown);
 		}
-		if (InColumnName == CompositePanelLayerTree::LayerTreeColumn_Name)
+		else if (InColumnName == CompositePanelLayerTree::LayerTreeColumn_Name)
 		{
 			return SNew(SHorizontalBox)
 
@@ -392,11 +415,11 @@ private:
 					.IsSelected(FIsSelected::CreateSP(this, &SCompositePanelLayerTreeItemRow::IsSelectedExclusively))
 				];
 		}
-		if (InColumnName == CompositePanelLayerTree::LayerTreeColumn_Type)
+		else if (InColumnName == CompositePanelLayerTree::LayerTreeColumn_Type)
 		{
 			return SNew(STextBlock).Text(this, &SCompositePanelLayerTreeItemRow::GetItemType);
 		}
-		if (InColumnName == CompositePanelLayerTree::LayerTreeColumn_Operation)
+		else if (InColumnName == CompositePanelLayerTree::LayerTreeColumn_Operation)
 		{
 			return SNew(STextBlock).Text(this, &SCompositePanelLayerTreeItemRow::GetItemOperation);	
 		}
@@ -409,6 +432,36 @@ private:
 		return TreeItem->HasValidCompositeActor() ? TreeItem->CompositeActor.Pin() : nullptr;
 	}
 
+	/** Gets the visibility state for the row item's Active flag */
+	EVisibility IsItemActiveVisible() const
+	{
+		// Only display active toggles on composite actor rows
+		const bool bIsVisible = TreeItem->HasValidCompositeActor() && !TreeItem->HasValidCompositeLayer();
+
+		return bIsVisible ? EVisibility::Visible : EVisibility::Hidden;
+	}
+
+	/** Gets the checkbox state for the row item's Active flag */
+	ECheckBoxState IsItemActive() const
+	{
+		if (TStrongObjectPtr<ACompositeActor> CompositeActor = GetCompositeActor())
+		{
+			return CompositeActor->IsActive() ? ECheckBoxState::Checked : ECheckBoxState::Unchecked;
+		}
+
+		return ECheckBoxState::Unchecked;
+	}
+
+	/** Raised when the row item's Active radiobox has changed state */
+	void OnIsItemActiveChanged(ECheckBoxState InCheckBoxState)
+	{
+		if (TStrongObjectPtr<ACompositeActor> CompositeActor = GetCompositeActor())
+		{
+			// No transactions, the property is non-transactional
+			CompositeActor->SetActive(InCheckBoxState == ECheckBoxState::Checked);
+		}
+	}
+
 	/** Gets the checkbox state for the row item's Enabled flag */
 	ECheckBoxState IsItemEnabled() const
 	{
@@ -766,7 +819,19 @@ void SCompositePanelLayerTree::Construct(const FArguments& InArgs)
 			.HeaderRow(
 				SNew(SHeaderRow)
 
-				// TODO: When multiple composite actors are supported, add a column here to allow a specific composite actor to be set to the active one
+				+SHeaderRow::Column(CompositePanelLayerTree::LayerTreeColumn_Active)
+				.DefaultLabel(FText::GetEmpty())
+				.FixedWidth(24.0f)
+				.HAlignHeader(HAlign_Center)
+				.VAlignHeader(VAlign_Center)
+				.HAlignCell(HAlign_Center)
+				.VAlignCell(VAlign_Center)
+				[
+					SNew(SCheckBox)
+						.Style(FAppStyle::Get(), "RadioButton")
+						.IsChecked(this, &SCompositePanelLayerTree::GetGlobalActiveCheckState)
+						.OnCheckStateChanged(this, &SCompositePanelLayerTree::OnGlobalEnabledCheckActiveChanged)
+				]
 				
 				+SHeaderRow::Column(CompositePanelLayerTree::LayerTreeColumn_Enabled)
 				.DefaultLabel(FText::GetEmpty())
@@ -1190,6 +1255,79 @@ void SCompositePanelLayerTree::OnPilotCameraToggled(bool bPilotCamera)
 	}
 }
 
+ECheckBoxState SCompositePanelLayerTree::GetGlobalActiveCheckState() const
+{
+	bool bAnyEnabled = false;
+	
+	for (const FCompositeActorTreeItemPtr& TreeItem : FilteredCompositeActorTreeItems)
+	{
+		if (!TreeItem.IsValid() || !TreeItem->HasValidCompositeActor() || TreeItem->HasValidCompositeLayer())
+		{
+			continue;
+		}
+
+		if (TreeItem->CompositeActor->IsActive())
+		{
+			bAnyEnabled = true;
+		}
+	}
+
+	return bAnyEnabled ? ECheckBoxState::Checked : ECheckBoxState::Unchecked;
+}
+
+void SCompositePanelLayerTree::OnGlobalEnabledCheckActiveChanged(ECheckBoxState CheckBoxState)
+{
+	if (CheckBoxState == ECheckBoxState::Checked)
+	{
+		bool bActivated = false;
+		for (const FCompositeActorTreeItemPtr& TreeItem : FilteredCompositeActorTreeItems)
+		{
+			if (!TreeItem.IsValid() || !TreeItem->HasValidCompositeActor() || TreeItem->HasValidCompositeLayer())
+			{
+				continue;
+			}
+
+			if (TreeItem->CompositeActor.Get() == LastActiveActor.Get())
+			{
+				TreeItem->CompositeActor->SetActive(true);
+				bActivated = true;
+			}
+		}
+
+		// If no match was found, we simply activate the first actor
+		if (!bActivated)
+		{
+			for (const FCompositeActorTreeItemPtr& TreeItem : FilteredCompositeActorTreeItems)
+			{
+				if (!TreeItem.IsValid() || !TreeItem->HasValidCompositeActor() || TreeItem->HasValidCompositeLayer())
+				{
+					continue;
+				}
+				
+				TreeItem->CompositeActor->SetActive(true);
+				break;
+			}
+		}
+	}
+	else
+	{
+		for (const FCompositeActorTreeItemPtr& TreeItem : FilteredCompositeActorTreeItems)
+		{
+			if (!TreeItem.IsValid() || !TreeItem->HasValidCompositeActor() || TreeItem->HasValidCompositeLayer())
+			{
+				continue;
+			}
+
+			if (TreeItem->CompositeActor->IsActive())
+			{
+				LastActiveActor = TreeItem->CompositeActor;
+			}
+
+			TreeItem->CompositeActor->SetActive(false);
+		}
+	}
+}
+
 ECheckBoxState SCompositePanelLayerTree::GetGlobalEnabledCheckState() const
 {
 	bool bAnyEnabled = false;
diff --git a/Engine/Plugins/Experimental/Compositing/Composite/Source/CompositeEditor/Private/UI/SCompositePanelLayerTree.h b/Engine/Plugins/Experimental/Compositing/Composite/Source/CompositeEditor/Private/UI/SCompositePanelLayerTree.h
--- a/Engine/Plugins/Experimental/Compositing/Composite/Source/CompositeEditor/Private/UI/SCompositePanelLayerTree.h
+++ b/Engine/Plugins/Experimental/Compositing/Composite/Source/CompositeEditor/Private/UI/SCompositePanelLayerTree.h
@@ -96,6 +96,12 @@ private:
 	/** Sets the piloted actor to the currently selected composite actor camera, or clears it if it is already being piloted */
 	void OnPilotCameraToggled(bool bPilotCamera);
 
+	/** Gets the global active checkbox state. */
+	ECheckBoxState GetGlobalActiveCheckState() const;
+
+	/** Sets the Enabled flag for all currently filtered items in the tree view */
+	void OnGlobalEnabledCheckActiveChanged(ECheckBoxState CheckBoxState);
+
 	/** Gets the check state for the global enabled checkbox */
 	ECheckBoxState GetGlobalEnabledCheckState() const;
 
@@ -171,6 +177,9 @@ private:
 
 	/** Raised when the selected tree items have changed */
 	FOnSelectionChanged OnSelectionChanged;
+
+	/** Last active actor. */
+	TWeakObjectPtr<ACompositeActor> LastActiveActor;
 	
 	friend class SCompositePanelLayerTreeItemRow;
 	friend class SCompositePanelLayerTreeToolbar;
diff --git a/Engine/Plugins/Experimental/Compositing/Composite/Source/CompositeEditor/Public/CompositeEditorModule.h b/Engine/Plugins/Experimental/Compositing/Composite/Source/CompositeEditor/Public/CompositeEditorModule.h
--- a/Engine/Plugins/Experimental/Compositing/Composite/Source/CompositeEditor/Public/CompositeEditorModule.h
+++ b/Engine/Plugins/Experimental/Compositing/Composite/Source/CompositeEditor/Public/CompositeEditorModule.h
@@ -7,6 +7,8 @@
 #include "Logging/LogMacros.h"
 #include "Modules/ModuleManager.h"
 
+class IConcertClientSession;
+
 /** Log category for this module. */
 DECLARE_LOG_CATEGORY_EXTERN(LogCompositeEditor, Log, All);
 
@@ -18,19 +20,29 @@ public:
 	virtual void StartupModule() override;
 	virtual void ShutdownModule() override;
 
+private:
 	/** Called after the engine has initialized. */
 	void OnPostEngineInit();
 
 	/** Called when an Actor has been added to a level. */
 	void OnLevelActorAdded(AActor* InActor);
 
-private:
+	/** Callback on multi-user concert session start. */
+	void OnConcertSessionStartup(TSharedRef<IConcertClientSession> InSession);
+
+	/** Register tab spawners. */
 	void RegisterTabSpawners();
+
+	/** Unregister tab spawners. */
 	void UnregisterTabSpawners();
 
 	void RegisterCustomizations();
 	void UnregisterCustomizations();
 
 private:
+	/** Cached concert sync client. */
+	TSharedPtr<class IConcertSyncClient> ConcertSyncClient;
+
+	/** Tab manager change delegate. */
 	FDelegateHandle OnTabManagerChangedDelegateHandle;
 };
