Composite: Naming hotfix from previous CL.

* Correctly replaced enabled with active term.

diff --git a/Engine/Plugins/Experimental/Compositing/Composite/Source/CompositeEditor/Private/UI/SCompositePanelLayerTree.cpp b/Engine/Plugins/Experimental/Compositing/Composite/Source/CompositeEditor/Private/UI/SCompositePanelLayerTree.cpp
--- a/Engine/Plugins/Experimental/Compositing/Composite/Source/CompositeEditor/Private/UI/SCompositePanelLayerTree.cpp
+++ b/Engine/Plugins/Experimental/Compositing/Composite/Source/CompositeEditor/Private/UI/SCompositePanelLayerTree.cpp
@@ -830,7 +830,7 @@ void SCompositePanelLayerTree::Construct(const FArguments& InArgs)
 					SNew(SCheckBox)
 						.Style(FAppStyle::Get(), "RadioButton")
 						.IsChecked(this, &SCompositePanelLayerTree::GetGlobalActiveCheckState)
-						.OnCheckStateChanged(this, &SCompositePanelLayerTree::OnGlobalEnabledCheckActiveChanged)
+						.OnCheckStateChanged(this, &SCompositePanelLayerTree::OnGlobalActiveCheckStateChanged)
 				]
 				
 				+SHeaderRow::Column(CompositePanelLayerTree::LayerTreeColumn_Enabled)
@@ -1275,7 +1275,7 @@ ECheckBoxState SCompositePanelLayerTree::GetGlobalActiveCheckState() const
 	return bAnyEnabled ? ECheckBoxState::Checked : ECheckBoxState::Unchecked;
 }
 
-void SCompositePanelLayerTree::OnGlobalEnabledCheckActiveChanged(ECheckBoxState CheckBoxState)
+void SCompositePanelLayerTree::OnGlobalActiveCheckStateChanged(ECheckBoxState CheckBoxState)
 {
 	if (CheckBoxState == ECheckBoxState::Checked)
 	{
diff --git a/Engine/Plugins/Experimental/Compositing/Composite/Source/CompositeEditor/Private/UI/SCompositePanelLayerTree.h b/Engine/Plugins/Experimental/Compositing/Composite/Source/CompositeEditor/Private/UI/SCompositePanelLayerTree.h
--- a/Engine/Plugins/Experimental/Compositing/Composite/Source/CompositeEditor/Private/UI/SCompositePanelLayerTree.h
+++ b/Engine/Plugins/Experimental/Compositing/Composite/Source/CompositeEditor/Private/UI/SCompositePanelLayerTree.h
@@ -99,8 +99,8 @@ private:
 	/** Gets the global active checkbox state. */
 	ECheckBoxState GetGlobalActiveCheckState() const;
 
-	/** Sets the Enabled flag for all currently filtered items in the tree view */
-	void OnGlobalEnabledCheckActiveChanged(ECheckBoxState CheckBoxState);
+	/** Sets the Active flag for all currently filtered items in the tree view */
+	void OnGlobalActiveCheckStateChanged(ECheckBoxState CheckBoxState);
 
 	/** Gets the check state for the global enabled checkbox */
 	ECheckBoxState GetGlobalEnabledCheckState() const;
