Composite: Fix missing concert sync client plugin dependency.

diff --git a/Engine/Plugins/Experimental/Compositing/Composite/Composite.uplugin b/Engine/Plugins/Experimental/Compositing/Composite/Composite.uplugin
--- a/Engine/Plugins/Experimental/Compositing/Composite/Composite.uplugin
+++ b/Engine/Plugins/Experimental/Compositing/Composite/Composite.uplugin
@@ -43,6 +43,10 @@
 		{
 			"Name": "MediaFrameworkUtilities",
 			"Enabled": true
+		},
+		{
+			"Name": "ConcertSyncClient",
+			"Enabled": true
 		}
 	]
 }
\ No newline at end of file
