Composite: Minor improvement to single light shadow layer.

diff --git a/Engine/Plugins/Experimental/Compositing/Composite/Shaders/Private/CompositeSingleLightShadow.usf b/Engine/Plugins/Experimental/Compositing/Composite/Shaders/Private/CompositeSingleLightShadow.usf
--- a/Engine/Plugins/Experimental/Compositing/Composite/Shaders/Private/CompositeSingleLightShadow.usf
+++ b/Engine/Plugins/Experimental/Compositing/Composite/Shaders/Private/CompositeSingleLightShadow.usf
@@ -16,6 +16,7 @@ float4x4        ShadowMatrix;
 float4          ShadowBufferSize;
 float4          ShadowInvDeviceZToWorldZ;
 float4          ShadowStrengthBias;
+float3          LightWorldDirection;
 
 uint            bIsPerspectiveProjection;
 
@@ -42,10 +43,9 @@ float ShadowDeviceZToSceneDepth(float DeviceZ)
 }
 
 // Percentage-close filtering shadow mapping
-float ManualPCF(in float4 ShadowPosition)
+float ManualPCF(in float4 ShadowPosition, in float NoL)
 {
-    // TODO: Improve shadow bias logic with normal & resolution info
-    float SceneDepth = ShadowDeviceZToSceneDepth(ShadowPosition.z) - ShadowStrengthBias.y;
+	float SceneDepth = ShadowDeviceZToSceneDepth(ShadowPosition.z) - saturate(1 - NoL) * ShadowStrengthBias.y;
 
     FPCFSamplerSettings Settings;
     Settings.ShadowDepthTexture = ShadowDepthTexture;
@@ -79,9 +79,15 @@ void MainPS(
     }
     else
     {
-        float ShadowVisibility = ManualPCF(ShadowPosition);
-        OutColor.rgb = lerp(1.0, ShadowVisibility, ShadowStrengthBias.x);
-    }
+        float3 Ddx = ddx(WorldPosition);
+        float3 Ddy = ddy(WorldPosition);
+        float3 WorldNormal = normalize(cross(Ddx, Ddy));
+        const float NoL = saturate(abs(dot(WorldNormal, LightWorldDirection)));
+
+        float ShadowVisibility = ManualPCF(ShadowPosition, NoL);
+
+		OutColor.rgb = lerp(1.0, ShadowVisibility, ShadowStrengthBias.x);
+	}
 
     OutColor.a = 1.0;
 }
diff --git a/Engine/Plugins/Experimental/Compositing/Composite/Source/Composite/Private/Layers/CompositeLayerSingleLightShadow.cpp b/Engine/Plugins/Experimental/Compositing/Composite/Source/Composite/Private/Layers/CompositeLayerSingleLightShadow.cpp
--- a/Engine/Plugins/Experimental/Compositing/Composite/Source/Composite/Private/Layers/CompositeLayerSingleLightShadow.cpp
+++ b/Engine/Plugins/Experimental/Compositing/Composite/Source/Composite/Private/Layers/CompositeLayerSingleLightShadow.cpp
@@ -32,6 +32,7 @@ class FCompositeSingleLightShadowShader : public FGlobalShader
 		SHADER_PARAMETER(FVector4f, ShadowBufferSize)
 		SHADER_PARAMETER(FVector4f, ShadowInvDeviceZToWorldZ)
 		SHADER_PARAMETER(FVector4f, ShadowStrengthBias)
+		SHADER_PARAMETER(FVector3f, LightWorldDirection)
 		SHADER_PARAMETER(uint32, bIsPerspectiveProjection)
 
 		RENDER_TARGET_BINDING_SLOTS()
@@ -127,7 +128,8 @@ namespace UE
 				Parameters->ShadowInvDeviceZToWorldZ	= ShadowInvDeviceZToWorldZ;
 				Parameters->ShadowMatrix				= ShadowMatrix;
 				Parameters->ShadowBufferSize			= FVector4f((float)ShadowMapSize.X, (float)ShadowMapSize.Y, 1.0f / (float)ShadowMapSize.X, 1.0f / (float)ShadowMapSize.Y);
-				Parameters->ShadowStrengthBias			= ShadowStrengthBias;
+				Parameters->ShadowStrengthBias		= ShadowStrengthBias;
+				Parameters->LightWorldDirection			= LightWorldDirection;
 				// We can use the current view since we the depth render matches the main view
 				Parameters->bIsPerspectiveProjection	= InView.IsPerspectiveProjection();
 				Parameters->RenderTargets[0]			= Output.GetRenderTargetBinding();
@@ -149,6 +151,7 @@ namespace UE
 			FMatrix44f ShadowMatrix;
 			FVector4f ShadowInvDeviceZToWorldZ;
 			FVector4f ShadowStrengthBias;
+			FVector3f LightWorldDirection;
 		};
 	}
 }
@@ -159,6 +162,17 @@ UCompositeLayerSingleLightShadow::UCompositeLayerSingleLightShadow(const FObject
 	Operation = ECompositeCoreMergeOp::Multiply;
 }
 
+void UCompositeLayerSingleLightShadow::SetEnabled(bool bInEnabled)
+{
+	Super::SetEnabled(bInEnabled);
+
+	ACompositeActor* CompositeActor = GetTypedOuter<ACompositeActor>();
+	if (IsValid(CompositeActor))
+	{
+		CompositeActor->ActivateSceneCaptures(this, IsEnabled());
+	}
+}
+
 void UCompositeLayerSingleLightShadow::BeginDestroy()
 {
 	Super::BeginDestroy();
@@ -269,10 +283,15 @@ bool UCompositeLayerSingleLightShadow::GetProxy(FTraversalContext& InContext, FS
 		PassInputs[0].Set<FPassExternalResourceDesc>({ EXTERNAL_RANGE_START_ID + SceneDepthIndex });
 		PassInputs[1].Set<FPassExternalResourceDesc>({ EXTERNAL_RANGE_START_ID + ShadowMapIndex });
 
+		static const TConsoleVariableData<float>* CVarDepthBias = IConsoleManager::Get().FindTConsoleVariableDataFloat(TEXT("r.Shadow.CSMDepthBias"));
+		float ShadowDepthBias = CVarDepthBias ? CVarDepthBias->GetValueOnGameThread() : 10.0f;
+		ShadowDepthBias *= ShadowBias;
+
 		ShadowProxy = InFrameAllocator.Create<FCompositeSingleLightShadowProxy>(PassInputs);
-		ShadowProxy->ShadowStrengthBias = FVector4f(ShadowStrength, ShadowBias, 0.0f, 0.0f);
+		ShadowProxy->ShadowStrengthBias = FVector4f(ShadowStrength, ShadowDepthBias, 0.0f, 0.0f);
 		GetShadowMatrices(*ShadowMapCapture, ShadowProxy->ShadowMatrix, ShadowProxy->ShadowInvDeviceZToWorldZ);
 	}
+	ShadowProxy->LightWorldDirection = FVector3f(Light->GetTransform().TransformVector(FVector(1, 0, 0)));
 
 	FPassInputDeclArray Inputs;
 	Inputs.SetNum(FixedNumLayerInputs);
diff --git a/Engine/Plugins/Experimental/Compositing/Composite/Source/Composite/Public/Layers/CompositeLayerSingleLightShadow.h b/Engine/Plugins/Experimental/Compositing/Composite/Source/Composite/Public/Layers/CompositeLayerSingleLightShadow.h
--- a/Engine/Plugins/Experimental/Compositing/Composite/Source/Composite/Public/Layers/CompositeLayerSingleLightShadow.h
+++ b/Engine/Plugins/Experimental/Compositing/Composite/Source/Composite/Public/Layers/CompositeLayerSingleLightShadow.h
@@ -37,11 +37,16 @@ public:
 	/** Destructor. */
 	UE_API virtual ~UCompositeLayerSingleLightShadow() = default;
 
+	/** Set the enabled state. */
+	UE_API virtual void SetEnabled(bool bInEnabled) override;
+
+	//~ Begin UObject Interface
 	UE_API virtual void BeginDestroy() override;
 
 #if WITH_EDITOR
 	UE_API virtual void PostEditChangeProperty(FPropertyChangedEvent& PropertyChangedEvent) override;
 #endif
+	//~ End UObject Interface
 
 	/** Getter function to override, returning pass layer proxies to be passed to the render thread. (Proxy objects should be allocated from the provided allocator.) */
 	UE_API virtual bool GetProxy(FTraversalContext& InContext, FSceneRenderingBulkObjectAllocator& InFrameAllocator, FCompositeCorePassProxy*& OutProxy) const override;
@@ -82,10 +87,10 @@ public:
 	float ShadowStrength = 0.5f;
 
 	/**
-	* Shadow bias (in world units) to avoid shadow acne. Increase until unwanted shadowing disappears.
+	* Shadow bias to avoid shadow acne. Increase until unwanted shadowing disappears (automatically scaled by r.Shadow.CSMDepthBias console variable).
 	*/
-	UPROPERTY(EditAnywhere, BlueprintReadWrite, Category = "", meta = (ClampMin = "0.0", UIMin = "0.0"))
-	float ShadowBias = 2.0f;
+	UPROPERTY(EditAnywhere, BlueprintReadWrite, Category = "", meta = (ClampMin = "0.0", UIMin = "0.0", UIMax = "1.0"))
+	float ShadowBias = 0.5f;
 
 private:
 	/**
