Composite: Convert actor pointers to soft object pointers to allow references in other levels.

* Refactor scene capture updates for hidden/show-only components. Since scene capture components are serialized, we only update these hidden/show-only properties when the layer properties are modified.
* Note that actor-picker customizations for single light shadow & scene capture layers are still needed.
* Auto-migration to soft pointers is backward compatible (ignoring blueprint).

diff --git a/Engine/Plugins/Experimental/Compositing/Composite/Source/Composite/Private/Layers/CompositeLayerPlate.cpp b/Engine/Plugins/Experimental/Compositing/Composite/Source/Composite/Private/Layers/CompositeLayerPlate.cpp
--- a/Engine/Plugins/Experimental/Compositing/Composite/Source/Composite/Private/Layers/CompositeLayerPlate.cpp
+++ b/Engine/Plugins/Experimental/Compositing/Composite/Source/Composite/Private/Layers/CompositeLayerPlate.cpp
@@ -39,11 +39,11 @@ namespace UE
 		namespace Private
 		{
 			/** Convenience function to remove composite asset user data from tracked mesh actors (before their removal). */
-			static void RemoveCompositeAssetUserData(TConstArrayView<AActor*> InPreEditMeshes, TConstArrayView<AActor*> InPostEditMeshes)
+			static void RemoveCompositeAssetUserData(TConstArrayView<TSoftObjectPtr<AActor>> InPreEditMeshes, TConstArrayView<TSoftObjectPtr<AActor>> InPostEditMeshes)
 			{
-				for (const AActor* PreEditMesh : InPreEditMeshes)
+				for (const TSoftObjectPtr<AActor>& PreEditMesh : InPreEditMeshes)
 				{
-					if (!IsValid(PreEditMesh))
+					if (!PreEditMesh.IsValid())
 					{
 						continue;
 					}
@@ -64,7 +64,7 @@ namespace UE
 			}
 
 			/** Convenience function to remove all composite asset user data from tracked mesh actors. */
-			static void RemoveCompositeAllAssetUserData(TConstArrayView<AActor*> InPreEditMeshes)
+			static void RemoveCompositeAllAssetUserData(TConstArrayView<TSoftObjectPtr<AActor>> InPreEditMeshes)
 			{
 				RemoveCompositeAssetUserData(InPreEditMeshes, {});
 			}
@@ -230,12 +230,12 @@ void UCompositeLayerPlate::BeginDestroy()
 	}
 }
 
-const TArray<AActor*> UCompositeLayerPlate::GetCompositeMeshes() const
+const TArray<TSoftObjectPtr<AActor>> UCompositeLayerPlate::GetCompositeMeshes() const
 {
 	return CompositeMeshes;
 }
 
-void UCompositeLayerPlate::SetCompositeMeshes(TArray<AActor*> InCompositeMeshes)
+void UCompositeLayerPlate::SetCompositeMeshes(TArray<TSoftObjectPtr<AActor>> InCompositeMeshes)
 {
 	// First, we untrack current meshes
 	UE::Composite::Private::RemoveCompositeAssetUserData(CompositeMeshes, InCompositeMeshes);
@@ -332,7 +332,7 @@ void UCompositeLayerPlate::PostEditChangeProperty(FPropertyChangedEvent& Propert
 				const int32 AlteredIndex = PropertyChangedEvent.GetArrayIndex(PropertyName.ToString());
 				if (CompositeMeshes.IsValidIndex(AlteredIndex))
 				{
-					TObjectPtr<AActor>& NewCompositeMesh = CompositeMeshes[AlteredIndex];
+					TSoftObjectPtr<AActor>& NewCompositeMesh = CompositeMeshes[AlteredIndex];
 
 					if (IsCompositeMeshActorAlreadyInUse(NewCompositeMesh))
 					{
@@ -543,9 +543,14 @@ TArray<UPrimitiveComponent*> UCompositeLayerPlate::GetPrimitives() const
 
 	TArray<UPrimitiveComponent*> OutPrimitiveComponents;
 
-	for (const AActor* Actor : CompositeMeshes)
+	for (const TSoftObjectPtr<AActor>& Actor : CompositeMeshes)
 	{
-		TInlineComponentArray<UPrimitiveComponent*> PrimComponents(Actor);
+		if (!Actor.IsValid())
+		{
+			continue;
+		}
+
+		TInlineComponentArray<UPrimitiveComponent*> PrimComponents(Actor.Get());
 
 		for (UPrimitiveComponent* PrimitiveComponent : PrimComponents)
 		{
@@ -566,9 +571,14 @@ int32 UCompositeLayerPlate::GetValidPrimitivesNum() const
 
 	int32 Count = 0;
 
-	for (const AActor* Actor : CompositeMeshes)
+	for (const TSoftObjectPtr<AActor>& Actor : CompositeMeshes)
 	{
-		TInlineComponentArray<UPrimitiveComponent*> PrimComponents(Actor);
+		if (!Actor.IsValid())
+		{
+			continue;
+		}
+
+		TInlineComponentArray<UPrimitiveComponent*> PrimComponents(Actor.Get());
 
 		for (UPrimitiveComponent* PrimitiveComponent : PrimComponents)
 		{
@@ -651,9 +661,9 @@ void UCompositeLayerPlate::OnPostSaveWorld(UWorld* InWorld, FObjectPostSaveConte
 
 void UCompositeLayerPlate::UpdateCompositeMeshes(bool bRegisterOnly) const
 {
-	for (AActor* MeshActor : CompositeMeshes)
+	for (const TSoftObjectPtr<AActor>& MeshActor : CompositeMeshes)
 	{
-		if (!IsValid(MeshActor))
+		if (!MeshActor.IsValid())
 		{
 			continue;
 		}
@@ -755,7 +765,7 @@ void UCompositeLayerPlate::HandleEnabledStateChange(bool bInEnabled, UCompositeC
 	}
 }
 
-bool UCompositeLayerPlate::IsCompositeMeshActorAlreadyInUse(AActor* InCompositeMeshActor) const
+bool UCompositeLayerPlate::IsCompositeMeshActorAlreadyInUse(TSoftObjectPtr<AActor> InCompositeMeshActor) const
 {
 	ACompositeActor* CompositeActor = GetTypedOuter<ACompositeActor>();
 	if (!IsValid(CompositeActor))
diff --git a/Engine/Plugins/Experimental/Compositing/Composite/Source/Composite/Private/Layers/CompositeLayerSceneCapture.cpp b/Engine/Plugins/Experimental/Compositing/Composite/Source/Composite/Private/Layers/CompositeLayerSceneCapture.cpp
--- a/Engine/Plugins/Experimental/Compositing/Composite/Source/Composite/Private/Layers/CompositeLayerSceneCapture.cpp
+++ b/Engine/Plugins/Experimental/Compositing/Composite/Source/Composite/Private/Layers/CompositeLayerSceneCapture.cpp
@@ -109,17 +109,37 @@ bool UCompositeLayerSceneCapture::GetProxy(FTraversalContext& InContext, FSceneR
 	return true;
 }
 
-const TArray<AActor*> UCompositeLayerSceneCapture::GetActors() const
+const TArray<TSoftObjectPtr<AActor>> UCompositeLayerSceneCapture::GetActors() const
 {
 	return Actors;
 }
 
-void UCompositeLayerSceneCapture::SetActors(TArray<AActor*> InActors)
+void UCompositeLayerSceneCapture::SetActors(TArray<TSoftObjectPtr<AActor>> InActors)
 {
 	USceneCaptureComponent2D* SceneCaptureComponent = GetSceneCaptureComponent();
 	if (IsValid(SceneCaptureComponent))
 	{
-		SceneCaptureComponent->ShowOnlyActors = InActors;
+		// First, we empty show-only actors since our previous logic used them
+		SceneCaptureComponent->ShowOnlyActors.Empty();
+
+		TArray<TWeakObjectPtr<UPrimitiveComponent>> ShowOnlyComponents;
+		ShowOnlyComponents.Reserve(InActors.Num());
+
+		for (const TSoftObjectPtr<AActor>& SoftActor : InActors)
+		{
+			if (const AActor* Actor = SoftActor.Get())
+			{
+				for (UActorComponent* Component : Actor->GetComponents())
+				{
+					if (UPrimitiveComponent* PrimitiveComponent = Cast<UPrimitiveComponent>(Component))
+					{
+						ShowOnlyComponents.Add(PrimitiveComponent);
+					}
+				}
+			}
+		}
+
+		SceneCaptureComponent->ShowOnlyComponents = ShowOnlyComponents;
 	}
 
 	Actors = MoveTemp(InActors);
diff --git a/Engine/Plugins/Experimental/Compositing/Composite/Source/Composite/Private/Layers/CompositeLayerShadowReflection.cpp b/Engine/Plugins/Experimental/Compositing/Composite/Source/Composite/Private/Layers/CompositeLayerShadowReflection.cpp
--- a/Engine/Plugins/Experimental/Compositing/Composite/Source/Composite/Private/Layers/CompositeLayerShadowReflection.cpp
+++ b/Engine/Plugins/Experimental/Compositing/Composite/Source/Composite/Private/Layers/CompositeLayerShadowReflection.cpp
@@ -6,6 +6,7 @@
 #include "Components/PrimitiveComponent.h"
 #include "Components/CompositeShadowReflectionCatcherComponent.h"
 #include "CompositeRenderTargetPool.h"
+#include "Containers/StaticArray.h"
 
 UCompositeLayerShadowReflection::UCompositeLayerShadowReflection(const FObjectInitializer& ObjectInitializer)
 	: Super(ObjectInitializer)
@@ -27,6 +28,20 @@ void UCompositeLayerShadowReflection::SetEnabled(bool bInEnabled)
 	}
 }
 
+void UCompositeLayerShadowReflection::BeginDestroy()
+{
+	Super::BeginDestroy();
+
+	if (USceneCaptureComponent2D* CaptureComponent = CachedSceneCapture_CG.Get())
+	{
+		FCompositeRenderTargetPool::Get().ReleaseTarget(CaptureComponent->TextureTarget);
+	}
+	if (USceneCaptureComponent2D* CaptureComponent = CachedSceneCapture_NoCG.Get())
+	{
+		FCompositeRenderTargetPool::Get().ReleaseTarget(CaptureComponent->TextureTarget);
+	}
+}
+
 #if WITH_EDITOR
 void UCompositeLayerShadowReflection::PostEditChangeProperty(FPropertyChangedEvent& PropertyChangedEvent)
 {
@@ -40,7 +55,7 @@ void UCompositeLayerShadowReflection::PostEditChangeProperty(FPropertyChangedEve
 	}
 	else if (PropertyName == GET_MEMBER_NAME_CHECKED(ThisClass, Actors))
 	{
-		ConfigurePrimitivesForHiddenShadowReflectionCatching();
+		SetActors(MoveTemp(Actors));
 	}
 	else if (PropertyName == GET_MEMBER_NAME_CHECKED(ThisClass, AutoConfigureActors))
 	{
@@ -59,24 +74,22 @@ bool UCompositeLayerShadowReflection::GetProxy(FTraversalContext& InContext, FSc
 		return false;
 	}
 
-	TStaticArray<UCompositeShadowReflectionCatcherComponent*, 2> ShadowReflectionPair;
-	ShadowReflectionPair[0] = CompositeActor->FindOrCreateSceneCapture<UCompositeShadowReflectionCatcherComponent>(this, 0, FName("ShadowReflectionCatcher_CG"));
-	ShadowReflectionPair[1] = CompositeActor->FindOrCreateSceneCapture<UCompositeShadowReflectionCatcherComponent>(this, 1, FName("ShadowReflectionCatcher_NoCG"));
-
-	FMergePassProxy* Proxy = nullptr;
-	FMergePassProxy* DivisionProxy = nullptr;
-
-	bool bValidComponents = true;
+	CachedSceneCapture_CG = FindOrCreateSceneCapture_CG(*CompositeActor);
+	CachedSceneCapture_NoCG = FindOrCreateSceneCapture_NoCG(*CompositeActor);
 
-	for (const UCompositeShadowReflectionCatcherComponent* Component : ShadowReflectionPair)
+	// Cached scene captures are now expected to be valid, early out if not.
+	if (!ensure(CachedSceneCapture_CG.IsValid() && CachedSceneCapture_NoCG.IsValid()))
 	{
-		bValidComponents &= IsValid(Component);
+		return false;
 	}
 
-	if(bValidComponents)
+	FMergePassProxy* DivisionProxy = nullptr;
+
 	{
-		// Update hidden actors
-		ShadowReflectionPair[1]->HiddenActors = Actors;
+		// Setup division proxy
+		TStaticArray<UCompositeShadowReflectionCatcherComponent*, 2> ShadowReflectionPair;
+		ShadowReflectionPair[0] = CachedSceneCapture_CG.Get();
+		ShadowReflectionPair[1] = CachedSceneCapture_NoCG.Get();
 
 		FPassInputDeclArray Inputs;
 		Inputs.SetNum(2);
@@ -98,11 +111,6 @@ bool UCompositeLayerShadowReflection::GetProxy(FTraversalContext& InContext, FSc
 
 		DivisionProxy = InFrameAllocator.Create<FMergePassProxy>(MoveTemp(Inputs), ECompositeCoreMergeOp::Divide, TEXT("ShadowRefl Div"), ELensDistortionHandling::Disabled);
 	}
-	
-	if (DivisionProxy == nullptr)
-	{
-		return false;
-	}
 
 	FPassInputDeclArray Inputs;
 	Inputs.SetNum(FixedNumLayerInputs);
@@ -113,16 +121,56 @@ bool UCompositeLayerShadowReflection::GetProxy(FTraversalContext& InContext, FSc
 	return true;
 }
 
-const TArray<AActor*> UCompositeLayerShadowReflection::GetActors() const
+const TArray<TSoftObjectPtr<AActor>> UCompositeLayerShadowReflection::GetActors() const
 {
 	return Actors;
 }
 
-void UCompositeLayerShadowReflection::SetActors(TArray<AActor*> InActors)
+void UCompositeLayerShadowReflection::SetActors(TArray<TSoftObjectPtr<AActor>> InActors)
 {
 	Actors = MoveTemp(InActors);
 
 	ConfigurePrimitivesForHiddenShadowReflectionCatching();
+
+	ACompositeActor* CompositeActor = GetTypedOuter<ACompositeActor>();
+	if (!IsValid(CompositeActor))
+	{
+		return;
+	}
+
+	if (!CachedSceneCapture_NoCG.IsValid())
+	{
+		CachedSceneCapture_NoCG = FindOrCreateSceneCapture_NoCG(*CompositeActor);
+	}
+
+	// Early exit if the cached scene depth capture is still not valid
+	if (!ensure(CachedSceneCapture_NoCG.IsValid()))
+	{
+		return;
+	}
+
+	// First, we empty hidden actors since our previous logic used them
+	CachedSceneCapture_NoCG->HiddenActors.Empty();
+
+	// Update hidden components
+	TArray<TWeakObjectPtr<UPrimitiveComponent> > HiddenComponents;
+	HiddenComponents.Reserve(Actors.Num());
+
+	for (const TSoftObjectPtr<AActor>& SoftActor : Actors)
+	{
+		if (const AActor* Actor = SoftActor.Get())
+		{
+			for (UActorComponent* Component : Actor->GetComponents())
+			{
+				if (UPrimitiveComponent* PrimitiveComponent = Cast<UPrimitiveComponent>(Component))
+				{
+					HiddenComponents.Add(PrimitiveComponent);
+				}
+			}
+		}
+	}
+
+	CachedSceneCapture_NoCG->HiddenComponents = HiddenComponents;
 }
 
 void UCompositeLayerShadowReflection::ConfigurePrimitivesForHiddenShadowReflectionCatching()
@@ -141,9 +189,9 @@ void UCompositeLayerShadowReflection::ConfigurePrimitivesForHiddenShadowReflecti
 		break;
 	}
 
-	for (AActor* Actor : Actors)
+	for (TSoftObjectPtr<AActor>& Actor : Actors)
 	{
-		if (IsValid(Actor))
+		if (Actor.IsValid())
 		{
 			if (UPrimitiveComponent* PrimitiveComponent = Actor->GetComponentByClass<UPrimitiveComponent>())
 			{
@@ -155,3 +203,13 @@ void UCompositeLayerShadowReflection::ConfigurePrimitivesForHiddenShadowReflecti
 	}
 }
 
+UCompositeShadowReflectionCatcherComponent* UCompositeLayerShadowReflection::FindOrCreateSceneCapture_CG(ACompositeActor& InOuter) const
+{
+	return InOuter.FindOrCreateSceneCapture<UCompositeShadowReflectionCatcherComponent>(this, 0, FName("ShadowReflectionCatcher_CG"));
+}
+
+UCompositeShadowReflectionCatcherComponent* UCompositeLayerShadowReflection::FindOrCreateSceneCapture_NoCG(ACompositeActor& InOuter) const
+{
+	return InOuter.FindOrCreateSceneCapture<UCompositeShadowReflectionCatcherComponent>(this, 1, FName("ShadowReflectionCatcher_NoCG"));
+}
+
diff --git a/Engine/Plugins/Experimental/Compositing/Composite/Source/Composite/Private/Layers/CompositeLayerSingleLightShadow.cpp b/Engine/Plugins/Experimental/Compositing/Composite/Source/Composite/Private/Layers/CompositeLayerSingleLightShadow.cpp
--- a/Engine/Plugins/Experimental/Compositing/Composite/Source/Composite/Private/Layers/CompositeLayerSingleLightShadow.cpp
+++ b/Engine/Plugins/Experimental/Compositing/Composite/Source/Composite/Private/Layers/CompositeLayerSingleLightShadow.cpp
@@ -183,7 +183,6 @@ void UCompositeLayerSingleLightShadow::BeginDestroy()
 	{
 		FCompositeRenderTargetPool::Get().ReleaseTarget(CaptureComponent->TextureTarget);
 	}
-
 	if (USceneCaptureComponent2D* CaptureComponent = CachedSceneDepthCapture.Get())
 	{
 		FCompositeRenderTargetPool::Get().ReleaseTarget(CaptureComponent->TextureTarget);
@@ -228,55 +227,42 @@ bool UCompositeLayerSingleLightShadow::GetProxy(FTraversalContext& InContext, FS
 		return false;
 	}
 
-	// Create new shadow map render scene capture (custom render pass) if needed
 	if (!CachedShadowMapCapture.IsValid())
 	{
-		CachedShadowMapCapture = CompositeActor->FindOrCreateSceneCapture<USceneCaptureComponent2D>(this, 0, FName("ShadowMapCapture"));
-		CachedShadowMapCapture->PrimitiveRenderMode = ESceneCapturePrimitiveRenderMode::PRM_RenderScenePrimitives;
-		CachedShadowMapCapture->CaptureSource = SCS_SceneDepth;
-		CachedShadowMapCapture->bRenderInMainRenderer = true;
-		CachedShadowMapCapture->bMainViewFamily = false;
-		CachedShadowMapCapture->bMainViewResolution = false;
-		CachedShadowMapCapture->bMainViewCamera = false;
-		CachedShadowMapCapture->bInheritMainViewCameraPostProcessSettings = false;
-		CachedShadowMapCapture->bUseRayTracingIfEnabled = false;
-		CachedShadowMapCapture->ProjectionType = ECameraProjectionMode::Orthographic;
+		CachedShadowMapCapture = FindOrCreateShadowMapCapture(*CompositeActor);
 	}
 
-	// Create new scene depth render without jitter
 	if (!CachedSceneDepthCapture.IsValid())
 	{
-		CachedSceneDepthCapture = CompositeActor->FindOrCreateSceneCapture<USceneCaptureComponent2D>(this, 1, FName("SceneDepthCapture"));
-		CachedSceneDepthCapture->PrimitiveRenderMode = ESceneCapturePrimitiveRenderMode::PRM_RenderScenePrimitives;
-		CachedSceneDepthCapture->CaptureSource = SCS_SceneDepth;
-		CachedSceneDepthCapture->bRenderInMainRenderer = true;
-		CachedSceneDepthCapture->bIgnoreScreenPercentage = true;
-		CachedSceneDepthCapture->HiddenActors = ShadowCastingActors;
+		CachedSceneDepthCapture = FindOrCreateSceneDepthCapture(*CompositeActor);
 	}
 
-	TStrongObjectPtr<USceneCaptureComponent2D> ShadowMapCapture = CachedShadowMapCapture.Pin();
-	TStrongObjectPtr<USceneCaptureComponent2D> SceneDepthCapture = CachedSceneDepthCapture.Pin();
+	// Both cached scene captures are now expected to be valid
+	if (!CachedShadowMapCapture.IsValid() || !CachedSceneDepthCapture.IsValid())
+	{
+		return false;
+	}
 
 	// Update properties that may be dynamic
 	{
 		// Keep exposed properties & transform in sync
-		ShadowMapCapture->OrthoWidth = OrthographicWidth;
-		ShadowMapCapture->SetWorldTransform(Light->GetTransform());
+		CachedShadowMapCapture->OrthoWidth = OrthographicWidth;
+		CachedShadowMapCapture->SetWorldTransform(Light->GetTransform());
 
 		// Acquire render target
-		FCompositeRenderTargetPool::Get().ConditionalAcquireTarget(ShadowMapCapture.Get(), ShadowMapCapture->TextureTarget, FIntPoint(ShadowMapResolution, ShadowMapResolution), RTF_R32f);
-		FCompositeRenderTargetPool::Get().ConditionalAcquireTarget(SceneDepthCapture.Get(), SceneDepthCapture->TextureTarget, GetRenderResolution(), RTF_R32f);
+		FCompositeRenderTargetPool::Get().ConditionalAcquireTarget(CachedShadowMapCapture.Get(), CachedShadowMapCapture->TextureTarget, FIntPoint(ShadowMapResolution, ShadowMapResolution), RTF_R32f);
+		FCompositeRenderTargetPool::Get().ConditionalAcquireTarget(CachedSceneDepthCapture.Get(), CachedSceneDepthCapture->TextureTarget, GetRenderResolution(), RTF_R32f);
 	}
 	
 	// Setup child shadow pass proxy
 	FCompositeSingleLightShadowProxy* ShadowProxy;
 	{
 		// Display-resolution scene depth without jitter as an external texture.
-		InContext.ExternalTextures.Add(FExternalTexture{ SceneDepthCapture->TextureTarget, FResourceMetadata{} });
+		InContext.ExternalTextures.Add(FExternalTexture{ CachedSceneDepthCapture->TextureTarget, FResourceMetadata{} });
 		const int32 SceneDepthIndex = InContext.ExternalTextures.Num() - 1;
 
 		// Shadow map texture
-		InContext.ExternalTextures.Add(FExternalTexture{ ShadowMapCapture->TextureTarget, FResourceMetadata{} });
+		InContext.ExternalTextures.Add(FExternalTexture{ CachedShadowMapCapture->TextureTarget, FResourceMetadata{} });
 		const int32 ShadowMapIndex = InContext.ExternalTextures.Num() - 1;
 
 		// Setup shadow pass inputs
@@ -291,7 +277,7 @@ bool UCompositeLayerSingleLightShadow::GetProxy(FTraversalContext& InContext, FS
 
 		ShadowProxy = InFrameAllocator.Create<FCompositeSingleLightShadowProxy>(PassInputs);
 		ShadowProxy->ShadowStrengthBias = FVector4f(ShadowStrength, ShadowDepthBias, 0.0f, 0.0f);
-		GetShadowMatrices(*ShadowMapCapture, ShadowProxy->ShadowMatrix, ShadowProxy->ShadowInvDeviceZToWorldZ);
+		GetShadowMatrices(*CachedShadowMapCapture, ShadowProxy->ShadowMatrix, ShadowProxy->ShadowInvDeviceZToWorldZ);
 	}
 	ShadowProxy->LightWorldDirection = FVector3f(Light->GetTransform().TransformVector(FVector(1, 0, 0)));
 
@@ -304,20 +290,89 @@ bool UCompositeLayerSingleLightShadow::GetProxy(FTraversalContext& InContext, FS
 	return true;
 }
 
-const TArray<AActor*> UCompositeLayerSingleLightShadow::GetShadowCastingActors() const
+const TArray<TSoftObjectPtr<AActor>> UCompositeLayerSingleLightShadow::GetShadowCastingActors() const
 {
 	return ShadowCastingActors;
 }
 
-void UCompositeLayerSingleLightShadow::SetShadowCastingActors(TArray<AActor*> InShadowCaster)
+void UCompositeLayerSingleLightShadow::SetShadowCastingActors(TArray<TSoftObjectPtr<AActor>> InShadowCaster)
 {
 	ShadowCastingActors = InShadowCaster;
 
-	TStrongObjectPtr<USceneCaptureComponent2D> SceneDepthCapture = CachedSceneDepthCapture.Pin();
-	if (SceneDepthCapture.IsValid())
+	ACompositeActor* CompositeActor = GetTypedOuter<ACompositeActor>();
+	if (!IsValid(CompositeActor))
+	{
+		return;
+	}
+
+	if (!CachedSceneDepthCapture.IsValid())
+	{
+		CachedSceneDepthCapture = FindOrCreateSceneDepthCapture(*CompositeActor);
+	}
+	
+	// Early exit if the cached scene depth capture is still not valid
+	if (!CachedSceneDepthCapture.IsValid())
 	{
-		SceneDepthCapture->HiddenActors = ShadowCastingActors;
+		return;
 	}
+
+	// First, we empty hidden actors since our previous logic used them
+	CachedSceneDepthCapture->HiddenActors.Empty();
+
+	// Collect & set hidden primitive components
+	TArray<TWeakObjectPtr<UPrimitiveComponent> > HiddenComponents;
+	HiddenComponents.Reserve(ShadowCastingActors.Num());
+
+	for (const TSoftObjectPtr<AActor>& SoftActor : ShadowCastingActors)
+	{
+		if (const AActor* Actor = SoftActor.Get())
+		{
+			for (UActorComponent* Component : Actor->GetComponents())
+			{
+				if (UPrimitiveComponent* PrimitiveComponent = Cast<UPrimitiveComponent>(Component))
+				{
+					HiddenComponents.Add(PrimitiveComponent);
+				}
+			}
+		}
+	}
+
+	CachedSceneDepthCapture->HiddenComponents = HiddenComponents;
+}
+
+USceneCaptureComponent2D* UCompositeLayerSingleLightShadow::FindOrCreateSceneDepthCapture(ACompositeActor& InOuter) const
+{
+	USceneCaptureComponent2D* SceneDepthCapture = InOuter.FindOrCreateSceneCapture<USceneCaptureComponent2D>(this, 1, FName("SceneDepthCapture"));
+	
+	if (ensure(IsValid(SceneDepthCapture)))
+	{
+		SceneDepthCapture->PrimitiveRenderMode = ESceneCapturePrimitiveRenderMode::PRM_RenderScenePrimitives;
+		SceneDepthCapture->CaptureSource = SCS_SceneDepth;
+		SceneDepthCapture->bRenderInMainRenderer = true;
+		SceneDepthCapture->bIgnoreScreenPercentage = true;
+	}
+
+	return SceneDepthCapture;
+}
+
+USceneCaptureComponent2D* UCompositeLayerSingleLightShadow::FindOrCreateShadowMapCapture(ACompositeActor& InOuter) const
+{
+	USceneCaptureComponent2D* ShadowMapCapture = InOuter.FindOrCreateSceneCapture<USceneCaptureComponent2D>(this, 0, FName("ShadowMapCapture"));
+
+	if (ensure(IsValid(ShadowMapCapture)))
+	{
+		ShadowMapCapture->PrimitiveRenderMode = ESceneCapturePrimitiveRenderMode::PRM_RenderScenePrimitives;
+		ShadowMapCapture->CaptureSource = SCS_SceneDepth;
+		ShadowMapCapture->bRenderInMainRenderer = true;
+		ShadowMapCapture->bMainViewFamily = false;
+		ShadowMapCapture->bMainViewResolution = false;
+		ShadowMapCapture->bMainViewCamera = false;
+		ShadowMapCapture->bInheritMainViewCameraPostProcessSettings = false;
+		ShadowMapCapture->bUseRayTracingIfEnabled = false;
+		ShadowMapCapture->ProjectionType = ECameraProjectionMode::Orthographic;
+	}
+
+	return ShadowMapCapture;
 }
 
 void UCompositeLayerSingleLightShadow::GetShadowMatrices(USceneCaptureComponent2D& ShadowMapCapture, FMatrix44f& OutShadowMatrix, FVector4f& OutShadowInvDeviceZToWorldZ) const
diff --git a/Engine/Plugins/Experimental/Compositing/Composite/Source/Composite/Public/Layers/CompositeLayerPlate.h b/Engine/Plugins/Experimental/Compositing/Composite/Source/Composite/Public/Layers/CompositeLayerPlate.h
--- a/Engine/Plugins/Experimental/Compositing/Composite/Source/Composite/Public/Layers/CompositeLayerPlate.h
+++ b/Engine/Plugins/Experimental/Compositing/Composite/Source/Composite/Public/Layers/CompositeLayerPlate.h
@@ -84,11 +84,11 @@ public:
 
 	/** Get composite meshes. */
 	UFUNCTION(BlueprintGetter)
-	UE_API const TArray<AActor*> GetCompositeMeshes() const;
+	UE_API const TArray<TSoftObjectPtr<AActor>> GetCompositeMeshes() const;
 
 	/** Set composite meshes. */
 	UFUNCTION(BlueprintSetter)
-	UE_API void SetCompositeMeshes(TArray<AActor*> InCompositeMeshes);
+	UE_API void SetCompositeMeshes(TArray<TSoftObjectPtr<AActor>> InCompositeMeshes);
 
 	/** Get the plate sampling & rendering mode */
 	UFUNCTION(BlueprintGetter)
@@ -105,7 +105,7 @@ public:
 private:
 	/** Actors whose primitive mesh components will be marked as primitive alpha holdout, and rendered separately with the projected (media) texture. */
 	UPROPERTY(EditAnywhere, BlueprintGetter = GetCompositeMeshes, BlueprintSetter = SetCompositeMeshes, Category = "", meta = (AllowPrivateAccess = true))
-	TArray<TObjectPtr<AActor>> CompositeMeshes;
+	TArray<TSoftObjectPtr<AActor>> CompositeMeshes;
 
 public:
 	/** Media texture for compositing layer. */
@@ -152,7 +152,7 @@ private:
 	void HandleEnabledStateChange(bool bInEnabled, UCompositeCoreSubsystem* InSubsystem) const;
 
 	/** Convenience function verify if the composite mesh actor is already in use by another layer. */
-	bool IsCompositeMeshActorAlreadyInUse(AActor* InCompositeMeshActor) const;
+	bool IsCompositeMeshActorAlreadyInUse(TSoftObjectPtr<AActor> InCompositeMeshActor) const;
 
 	/** Convenience function to get or create the specified render target at composite actor resolution. */
 	TObjectPtr<UTextureRenderTarget2D> GetOrCreateRenderTarget(TObjectPtr<UTextureRenderTarget2D>& InRenderTarget) const;
@@ -185,7 +185,7 @@ private:
 
 #if WITH_EDITOR
 	/** Pre-edit list of registered composite meshes. */
-	TArray<TObjectPtr<AActor>> PreEditCompositeMeshes;
+	TArray<TSoftObjectPtr<AActor>> PreEditCompositeMeshes;
 
 	/**
 	 * Called before a level saves.
diff --git a/Engine/Plugins/Experimental/Compositing/Composite/Source/Composite/Public/Layers/CompositeLayerSceneCapture.h b/Engine/Plugins/Experimental/Compositing/Composite/Source/Composite/Public/Layers/CompositeLayerSceneCapture.h
--- a/Engine/Plugins/Experimental/Compositing/Composite/Source/Composite/Public/Layers/CompositeLayerSceneCapture.h
+++ b/Engine/Plugins/Experimental/Compositing/Composite/Source/Composite/Public/Layers/CompositeLayerSceneCapture.h
@@ -41,11 +41,11 @@ public:
 
 	/** Get reflection & shadow caster actors. */
 	UFUNCTION(BlueprintGetter)
-	UE_API const TArray<AActor*> GetActors() const;
+	UE_API const TArray<TSoftObjectPtr<AActor>> GetActors() const;
 
 	/** Set reflection & shadow caster actors. */
 	UFUNCTION(BlueprintSetter)
-	UE_API void SetActors(TArray<AActor*> InActors);
+	UE_API void SetActors(TArray<TSoftObjectPtr<AActor>> InActors);
 
 	UFUNCTION(BlueprintGetter)
 	UE_API bool IsCustomRenderPass() const;
@@ -64,10 +64,10 @@ private:
 	/**
 	 * List of actors to be rendered from the scene capture.
 	 * 
-	 * Automatically sets the scene capture's ShowOnlyActors list when updated.
+	 * Automatically sets the scene capture's ShowOnlyComponents when updated.
 	*/
 	UPROPERTY(EditAnywhere, BlueprintGetter = GetActors, BlueprintSetter = SetActors, Category = "", meta = (AllowPrivateAccess = true))
-	TArray<TObjectPtr<AActor>> Actors;
+	TArray<TSoftObjectPtr<AActor>> Actors;
 
 	/**
 	 * Convenience toggle to render the scene capture as a fast/minimal Custom Render Pass, inlined in the main render & without support for lighting.
diff --git a/Engine/Plugins/Experimental/Compositing/Composite/Source/Composite/Public/Layers/CompositeLayerShadowReflection.h b/Engine/Plugins/Experimental/Compositing/Composite/Source/Composite/Public/Layers/CompositeLayerShadowReflection.h
--- a/Engine/Plugins/Experimental/Compositing/Composite/Source/Composite/Public/Layers/CompositeLayerShadowReflection.h
+++ b/Engine/Plugins/Experimental/Compositing/Composite/Source/Composite/Public/Layers/CompositeLayerShadowReflection.h
@@ -5,12 +5,16 @@
 #include "CompositeLayerBase.h"
 
 #include "Engine/TextureRenderTarget2D.h"
+#include "Containers/ContainersFwd.h"
 
 #include "CompositeLayerShadowReflection.generated.h"
 
 #define UE_API COMPOSITE_API
 
 class AActor;
+class ACompositeActor;
+class UCompositeShadowReflectionCatcherComponent;
+class USceneCaptureComponent2D;
 
 /* Automatic primitive configuration mode for registered primitives. */
 UENUM(BlueprintType)
@@ -40,20 +44,24 @@ public:
 	/** Set the enabled state. */
 	UE_API virtual void SetEnabled(bool bInEnabled) override;
 
+	//~ Begin UObject Interface
+	UE_API virtual void BeginDestroy() override;
+
 #if WITH_EDITOR
 	UE_API virtual void PostEditChangeProperty(FPropertyChangedEvent& PropertyChangedEvent) override;
 #endif
+	//~ End UObject Interface
 
 	/** Getter function to override, returning pass layer proxies to be passed to the render thread. (Proxy objects should be allocated from the provided allocator.) */
 	UE_API virtual bool GetProxy(FTraversalContext& InContext, FSceneRenderingBulkObjectAllocator& InFrameAllocator, FCompositeCorePassProxy*& OutProxy) const override;
 
 	/** Get reflection & shadow caster actors. */
 	UFUNCTION(BlueprintGetter)
-	UE_API const TArray<AActor*> GetActors() const;
+	UE_API const TArray<TSoftObjectPtr<AActor>> GetActors() const;
 
 	/** Set reflection & shadow caster actors. */
 	UFUNCTION(BlueprintSetter)
-	UE_API void SetActors(TArray<AActor*> InActors);
+	UE_API void SetActors(TArray<TSoftObjectPtr<AActor>> InActors);
 
 	/**
 	 * Depending on the configuration mode, we automatically enable or disable the following (primitive component) properties:
@@ -69,9 +77,22 @@ private:
 	/** Convenience function to configure actor primitives. */
 	void ConfigurePrimitivesForHiddenShadowReflectionCatching();
 
+	/** Find or create scene capture with CG. */
+	UCompositeShadowReflectionCatcherComponent* FindOrCreateSceneCapture_CG(ACompositeActor& InOuter) const;
+
+	/** Find or create scene capture without CG. */
+	UCompositeShadowReflectionCatcherComponent* FindOrCreateSceneCapture_NoCG(ACompositeActor& InOuter) const;
+
+private:
 	/** List of reflection & shadow caster actors. */
 	UPROPERTY(EditAnywhere, BlueprintGetter = GetActors, BlueprintSetter = SetActors, Category = "")
-	TArray<TObjectPtr<AActor>> Actors;
+	TArray<TSoftObjectPtr<AActor>> Actors;
+
+	/** Cached pointer to the scene capture with CG. */
+	mutable TWeakObjectPtr<UCompositeShadowReflectionCatcherComponent> CachedSceneCapture_CG;
+
+	/** Cached pointer to the scene capture without CG. */
+	mutable TWeakObjectPtr<UCompositeShadowReflectionCatcherComponent> CachedSceneCapture_NoCG;
 
 	friend class FCompositeLayerShadowReflectionCustomization;
 };
diff --git a/Engine/Plugins/Experimental/Compositing/Composite/Source/Composite/Public/Layers/CompositeLayerSingleLightShadow.h b/Engine/Plugins/Experimental/Compositing/Composite/Source/Composite/Public/Layers/CompositeLayerSingleLightShadow.h
--- a/Engine/Plugins/Experimental/Compositing/Composite/Source/Composite/Public/Layers/CompositeLayerSingleLightShadow.h
+++ b/Engine/Plugins/Experimental/Compositing/Composite/Source/Composite/Public/Layers/CompositeLayerSingleLightShadow.h
@@ -63,11 +63,11 @@ public:
 
 	/** Get shadow caster actors. */
 	UFUNCTION(BlueprintGetter)
-	UE_API const TArray<AActor*> GetShadowCastingActors() const;
+	UE_API const TArray<TSoftObjectPtr<AActor>> GetShadowCastingActors() const;
 
 	/** Set shadow caster actors. */
 	UFUNCTION(BlueprintSetter)
-	UE_API void SetShadowCastingActors(TArray<AActor*> InShadowCaster);
+	UE_API void SetShadowCastingActors(TArray<TSoftObjectPtr<AActor>> InShadowCaster);
 
 	/**
 	* The desired width (in world units) of the shadow map view (ignored if light is not directional).
@@ -92,16 +92,23 @@ public:
 	UPROPERTY(EditAnywhere, BlueprintReadWrite, Category = "", meta = (ClampMin = "0.0", UIMin = "0.0", UIMax = "1.0"))
 	float ShadowBias = 0.5f;
 
+private:
+	/** Find or create scene depth capture component, correctly configured. */
+	USceneCaptureComponent2D* FindOrCreateSceneDepthCapture(ACompositeActor& InOuter) const;
+
+	/** Find or create shadow map capture component, correctly configured. */
+	USceneCaptureComponent2D* FindOrCreateShadowMapCapture(ACompositeActor& InOuter) const;
+
+	/** Convenience function to calculate shadow matrices. */
+	void GetShadowMatrices(USceneCaptureComponent2D& ShadowMapCapture, FMatrix44f& OutShadowMatrix, FVector4f& OutShadowInvDeviceZToWorldZ) const;
+
 private:
 	/**
 	* Optional: List of shadow casting actors to be removed from the main camera view's scene depth.
 	* This helps create a cleaner shadow matte preventing aliased dark edges around CG objects.
 	*/
 	UPROPERTY(EditAnywhere, BlueprintGetter = GetShadowCastingActors, BlueprintSetter = SetShadowCastingActors, Category = "", meta = (AllowPrivateAccess = true))
-	TArray<TObjectPtr<AActor>> ShadowCastingActors;
-
-	/** Convenience function to calculate shadow matrices. */
-	void GetShadowMatrices(USceneCaptureComponent2D& ShadowMapCapture, FMatrix44f& OutShadowMatrix, FVector4f& OutShadowInvDeviceZToWorldZ) const;
+	TArray<TSoftObjectPtr<AActor>> ShadowCastingActors;
 
 	/** Cached pointer to the composite-managed shadow map scene capture. */
 	mutable TWeakObjectPtr<USceneCaptureComponent2D> CachedShadowMapCapture;
diff --git a/Engine/Plugins/Experimental/Compositing/Composite/Source/CompositeEditor/Private/UI/SCompositeActorPickerTable.cpp b/Engine/Plugins/Experimental/Compositing/Composite/Source/CompositeEditor/Private/UI/SCompositeActorPickerTable.cpp
--- a/Engine/Plugins/Experimental/Compositing/Composite/Source/CompositeEditor/Private/UI/SCompositeActorPickerTable.cpp
+++ b/Engine/Plugins/Experimental/Compositing/Composite/Source/CompositeEditor/Private/UI/SCompositeActorPickerTable.cpp
@@ -323,7 +323,7 @@ private:
 	FActorListItemRef ActorItem;
 };
 
-FCompositeActorPickerListRef::FCompositeActorPickerListRef(const TWeakObjectPtr<UObject>& InActorListOwner, const FName& InActorListPropertyName, TArray<TObjectPtr<AActor>>* InActorList)
+FCompositeActorPickerListRef::FCompositeActorPickerListRef(const TWeakObjectPtr<UObject>& InActorListOwner, const FName& InActorListPropertyName, TArray<TSoftObjectPtr<AActor>>* InActorList)
 	: ActorListOwner(InActorListOwner)
 	, ActorListPropertyName(InActorListPropertyName)
 	, ActorList(InActorList)
diff --git a/Engine/Plugins/Experimental/Compositing/Composite/Source/CompositeEditor/Private/UI/SCompositeActorPickerTable.h b/Engine/Plugins/Experimental/Compositing/Composite/Source/CompositeEditor/Private/UI/SCompositeActorPickerTable.h
--- a/Engine/Plugins/Experimental/Compositing/Composite/Source/CompositeEditor/Private/UI/SCompositeActorPickerTable.h
+++ b/Engine/Plugins/Experimental/Compositing/Composite/Source/CompositeEditor/Private/UI/SCompositeActorPickerTable.h
@@ -14,7 +14,7 @@ template<typename T> class SListView;
 struct FCompositeActorPickerListRef
 {
 public:
-	FCompositeActorPickerListRef(const TWeakObjectPtr<UObject>& InActorListOwner, const FName& InActorListPropertyName, TArray<TObjectPtr<AActor>>* InActorList);
+	FCompositeActorPickerListRef(const TWeakObjectPtr<UObject>& InActorListOwner, const FName& InActorListPropertyName, TArray<TSoftObjectPtr<AActor>>* InActorList);
 	
 	FCompositeActorPickerListRef() { }
 	
@@ -38,7 +38,7 @@ public:
 	FProperty* ActorListProperty = nullptr;
 	
 	/** Pointer to the actor list being edited by this widget */
-	TArray<TObjectPtr<AActor>>* ActorList = nullptr;
+	TArray<TSoftObjectPtr<AActor>>* ActorList = nullptr;
 };
 
 /**
@@ -49,13 +49,13 @@ class SCompositeActorPickerTable : public SCompoundWidget, public FEditorUndoCli
 private:
 	struct FActorListItem
 	{
-		/** Weak reference to the actor this item represents */
-		TWeakObjectPtr<AActor> Actor;
+		/** Soft reference to the actor this item represents */
+		TSoftObjectPtr<AActor> Actor;
 
 		/** The index of the actor in the list */
 		int32 Index;
 
-		FActorListItem(const TWeakObjectPtr<AActor>& InActor, int32 InIndex)
+		FActorListItem(const TSoftObjectPtr<AActor>& InActor, int32 InIndex)
 			: Actor(InActor)
 			, Index(InIndex)
 		{ }
