Composite: Disabled despill by default to prevent unexpected chromaticity changes on unkeyed colors.

Users should now opt-in when using green or blue stages, but less so for simulcam.

diff --git a/Engine/Plugins/Experimental/Compositing/Composite/Source/Composite/Private/Passes/CompositePassDeltaKeyer.cpp b/Engine/Plugins/Experimental/Compositing/Composite/Source/Composite/Private/Passes/CompositePassDeltaKeyer.cpp
--- a/Engine/Plugins/Experimental/Compositing/Composite/Source/Composite/Private/Passes/CompositePassDeltaKeyer.cpp
+++ b/Engine/Plugins/Experimental/Compositing/Composite/Source/Composite/Private/Passes/CompositePassDeltaKeyer.cpp
@@ -208,7 +208,7 @@ UCompositePassDeltaKeyer::UCompositePassDeltaKeyer(const FObjectInitializer& Obj
 	, GreenWeight(0.5f)
 	, BlueWeight(0.5f)
 	, AlphaThreshold(0.0f, 1.0f)
-	, DespillStrength(1.0f)
+	, DespillStrength(0.0f)
 	, DevignetteStrength(0.0f)
 	, bPreserveVignetteAfterKey(true)
 	, DenoiseMethod(ECompositeDenoiseMethod::None)
