Composite: Remove AM_Border use, which causes unwanted translucency on the edges.

diff --git a/Engine/Plugins/Experimental/Compositing/Composite/Source/Composite/Private/Passes/CompositePassCenteredScale.cpp b/Engine/Plugins/Experimental/Compositing/Composite/Source/Composite/Private/Passes/CompositePassCenteredScale.cpp
--- a/Engine/Plugins/Experimental/Compositing/Composite/Source/Composite/Private/Passes/CompositePassCenteredScale.cpp
+++ b/Engine/Plugins/Experimental/Compositing/Composite/Source/Composite/Private/Passes/CompositePassCenteredScale.cpp
@@ -59,7 +59,7 @@ namespace UE
 
 				FCompositeCenteredScaleShader::FParameters* Parameters = GraphBuilder.AllocParameters<FCompositeCenteredScaleShader::FParameters>();
 				Parameters->InputTexture = Input.Texture;
-				Parameters->InputSampler = TStaticSamplerState<SF_Bilinear, AM_Border, AM_Border, AM_Clamp>::GetRHI();
+				Parameters->InputSampler = TStaticSamplerState<SF_Bilinear>::GetRHI();
 				Parameters->ScaleUV = ScaleUV;
 				Parameters->RenderTargets[0] = Output.GetRenderTargetBinding();
 
diff --git a/Engine/Plugins/Experimental/Compositing/Composite/Source/Composite/Private/Passes/CompositePassColorGrade.cpp b/Engine/Plugins/Experimental/Compositing/Composite/Source/Composite/Private/Passes/CompositePassColorGrade.cpp
--- a/Engine/Plugins/Experimental/Compositing/Composite/Source/Composite/Private/Passes/CompositePassColorGrade.cpp
+++ b/Engine/Plugins/Experimental/Compositing/Composite/Source/Composite/Private/Passes/CompositePassColorGrade.cpp
@@ -84,7 +84,7 @@ public:
 
 		FCompositeColorGradeShader::FParameters* Parameters = GraphBuilder.AllocParameters<FCompositeColorGradeShader::FParameters>();
 		Parameters->InputTexture = Input.Texture;
-		Parameters->InputSampler = TStaticSamplerState<SF_Bilinear, AM_Border, AM_Border, AM_Clamp>::GetRHI();
+		Parameters->InputSampler = TStaticSamplerState<SF_Bilinear>::GetRHI();
 		Parameters->WorkingColorSpace = GDefaultWorkingColorSpaceUniformBuffer.GetUniformBufferRef();
 		Parameters->Global = FCompositeColorGradeShader::ToParameters(ColorGradingSettings.Global);
 		Parameters->Shadows = FCompositeColorGradeShader::ToParameters(ColorGradingSettings.Shadows);
