Composite: Fix lens distortion in scene captures by providing the correct view owner.

* By default, scene capture's view owner is null which breaks our lens distortion SVE since it relies on the view actor unique id.
* All composite actor scene captures are now custom derived classes.

diff --git a/Engine/Plugins/Experimental/Compositing/Composite/Source/Composite/Private/Components/CompositeSceneCapture2DComponent.cpp b/Engine/Plugins/Experimental/Compositing/Composite/Source/Composite/Private/Components/CompositeSceneCapture2DComponent.cpp
new file mode 100644
--- /dev/null
+++ b/Engine/Plugins/Experimental/Compositing/Composite/Source/Composite/Private/Components/CompositeSceneCapture2DComponent.cpp
@@ -0,0 +1,78 @@
+// Copyright Epic Games, Inc. All Rights Reserved.
+
+#include "Components/CompositeSceneCapture2DComponent.h"
+
+#include "Camera/CameraComponent.h"
+#include "CompositeActor.h"
+#include "CompositeRenderTargetPool.h"
+#include "Engine/TextureRenderTarget2D.h"
+#include "UObject/Package.h"
+
+#define LOCTEXT_NAMESPACE "Composite"
+
+UCompositeSceneCapture2DComponent::UCompositeSceneCapture2DComponent(const FObjectInitializer& ObjectInitializer)
+	: Super(ObjectInitializer)
+{
+	PrimitiveRenderMode = ESceneCapturePrimitiveRenderMode::PRM_UseShowOnlyList;
+
+	CaptureSource = ESceneCaptureSource::SCS_FinalColorHDR;
+	bUseRayTracingIfEnabled = true;
+	bMainViewFamily = true;
+	bMainViewResolution = true;
+	bMainViewCamera = true;
+	// Inherit main view camera post-process settings to match main view DoF for example.
+	bInheritMainViewCameraPostProcessSettings = true;
+	bIgnoreScreenPercentage = true;
+	bConsiderUnrenderedOpaquePixelAsFullyTranslucent = true;
+	UnlitViewmode = ESceneCaptureUnlitViewmode::CaptureOrCustomRenderPass;
+
+	PostProcessSettings.bOverride_ReflectionMethod = true;
+	PostProcessSettings.ReflectionMethod = EReflectionMethod::Lumen;
+	PostProcessSettings.bOverride_DynamicGlobalIlluminationMethod = true;
+	PostProcessSettings.DynamicGlobalIlluminationMethod = EDynamicGlobalIlluminationMethod::Lumen;
+	PostProcessSettings.bOverride_LumenFinalGatherScreenTraces = true;
+	PostProcessSettings.LumenFinalGatherScreenTraces = false;
+	PostProcessSettings.bOverride_LumenReflectionsScreenTraces = true;
+	PostProcessSettings.LumenReflectionsScreenTraces = false;
+	// While hit lighting is more expensive, it significantly minimizes artifacts & noise. Here we default with higher-quality settings.
+	PostProcessSettings.bOverride_LumenRayLightingMode = true;
+	PostProcessSettings.LumenRayLightingMode = ELumenRayLightingModeOverride::HitLighting;
+
+	ShowFlags.SetAtmosphere(false);
+	ShowFlags.SetAllowPrimitiveAlphaHoldout(false);
+	ShowFlags.SetBloom(false);
+	ShowFlags.SetEyeAdaptation(false);
+	ShowFlags.SetLocalExposure(false);
+	ShowFlags.SetPostProcessMaterial(false);
+
+}
+
+void UCompositeSceneCapture2DComponent::BeginDestroy()
+{
+	FCompositeRenderTargetPool::Get().ReleaseTarget(TextureTarget);
+
+	Super::BeginDestroy();
+}
+
+#if WITH_EDITOR
+void UCompositeSceneCapture2DComponent::PostEditChangeProperty(FPropertyChangedEvent& PropertyChangedEvent)
+{
+	Super::PostEditChangeProperty(PropertyChangedEvent);
+}
+#endif
+
+const AActor* UCompositeSceneCapture2DComponent::GetViewOwner() const
+{
+	// View owner that matches the main view camera
+	if (ACompositeActor* CompositeActor = GetTypedOuter<ACompositeActor>())
+	{
+		if (UCameraComponent* CameraComponent = Cast<UCameraComponent>(CompositeActor->GetCamera().GetComponent(nullptr)))
+		{
+			return CameraComponent->GetOwner();
+		}
+	}
+
+	return nullptr;
+}
+
+#undef LOCTEXT_NAMESPACE
diff --git a/Engine/Plugins/Experimental/Compositing/Composite/Source/Composite/Private/Components/CompositeShadowReflectionCatcherComponent.cpp b/Engine/Plugins/Experimental/Compositing/Composite/Source/Composite/Private/Components/CompositeShadowReflectionCatcherComponent.cpp
--- a/Engine/Plugins/Experimental/Compositing/Composite/Source/Composite/Private/Components/CompositeShadowReflectionCatcherComponent.cpp
+++ b/Engine/Plugins/Experimental/Compositing/Composite/Source/Composite/Private/Components/CompositeShadowReflectionCatcherComponent.cpp
@@ -3,7 +3,6 @@
 #include "Components/CompositeShadowReflectionCatcherComponent.h"
 
 #include "CompositeRenderTargetPool.h"
-#include "CompositeUtilities.h"
 #include "Engine/TextureRenderTarget2D.h"
 #include "UObject/Package.h"
 
@@ -12,32 +11,12 @@
 UCompositeShadowReflectionCatcherComponent::UCompositeShadowReflectionCatcherComponent(const FObjectInitializer& ObjectInitializer)
 	: Super(ObjectInitializer)
 {
-	UE::Composite::ConfigureSceneCapture(*this, ESceneCapturePrimitiveRenderMode::PRM_RenderScenePrimitives);
+	PrimitiveRenderMode = ESceneCapturePrimitiveRenderMode::PRM_RenderScenePrimitives;
 
-	ShowFlags.SetAtmosphere(false);
-	ShowFlags.SetAllowPrimitiveAlphaHoldout(false);
 	ShowFlags.SetGrain(false);
-	ShowFlags.SetBloom(false);
-	ShowFlags.SetEyeAdaptation(false);
-	ShowFlags.SetLocalExposure(false);
-	ShowFlags.SetPostProcessMaterial(false);
 	ShowFlags.SetScreenSpaceReflections(false);
 	//Note: This can cause temporal ghosting artifacts on moving shadows, but remains preferable over the otherwise default shadow noise.
 	ShowFlags.SetTemporalAA(true);
 }
 
-void UCompositeShadowReflectionCatcherComponent::BeginDestroy()
-{
-	FCompositeRenderTargetPool::Get().ReleaseTarget(TextureTarget);
-
-	Super::BeginDestroy();
-}
-
-#if WITH_EDITOR
-void UCompositeShadowReflectionCatcherComponent::PostEditChangeProperty(FPropertyChangedEvent& PropertyChangedEvent)
-{
-	Super::PostEditChangeProperty(PropertyChangedEvent);
-}
-#endif
-
 #undef LOCTEXT_NAMESPACE
diff --git a/Engine/Plugins/Experimental/Compositing/Composite/Source/Composite/Private/CompositeActor.cpp b/Engine/Plugins/Experimental/Compositing/Composite/Source/Composite/Private/CompositeActor.cpp
--- a/Engine/Plugins/Experimental/Compositing/Composite/Source/Composite/Private/CompositeActor.cpp
+++ b/Engine/Plugins/Experimental/Compositing/Composite/Source/Composite/Private/CompositeActor.cpp
@@ -5,7 +5,6 @@
 #include "CompositeCoreSubsystem.h"
 #include "CompositeModule.h"
 #include "CompositeCoreSettings.h"
-#include "CompositeUtilities.h"
 #include "CompositeRenderTargetPool.h"
 
 #include "Components/CompositeShadowReflectionCatcherComponent.h"
@@ -26,6 +25,7 @@
 #include "Materials/MaterialInstanceDynamic.h"
 #include "SceneView.h"
 #include "UObject/FortniteMainBranchObjectVersion.h"
+#include "UObject/UE5ReleaseStreamObjectVersion.h"
 #include "UObject/ObjectSaveContext.h"
 #include "UObject/UObjectIterator.h"
 
@@ -243,21 +243,39 @@ void ACompositeActor::PostLoad()
 
 #if WITH_EDITORONLY_DATA
 	/**
-	* Refactored scene capture management where previous shadow/reflection maps were removed.
+	* Refactored scene capture management.
 	* Instead of polluting the actor with old data structures, we remove old components and rely on their recreation.
 	*/
-	if (GetLinkerCustomVersion(FFortniteMainBranchObjectVersion::GUID) < FFortniteMainBranchObjectVersion::CompositeActorSceneCaptureRefactor)
+	const bool bHandleSceneCaptureRefactor = GetLinkerCustomVersion(FFortniteMainBranchObjectVersion::GUID) < FFortniteMainBranchObjectVersion::CompositeActorSceneCaptureRefactor;
+	const bool bHandleDerivedSceneCaptureTypes = GetLinkerCustomVersion(FUE5ReleaseStreamObjectVersion::GUID) < FUE5ReleaseStreamObjectVersion::CompositePluginDerivedSceneCaptures;
+
+	if (bHandleSceneCaptureRefactor || bHandleDerivedSceneCaptureTypes)
 	{
 		TArray<UActorComponent*> ActorComponents;
 		GetComponents(ActorComponents);
 
 		for (auto It = ActorComponents.CreateIterator(); It; ++It)
 		{
-			if (UCompositeShadowReflectionCatcherComponent* CaptureComponent = Cast<UCompositeShadowReflectionCatcherComponent>(*It))
+			if (USceneCaptureComponent2D* CaptureComponent = Cast<USceneCaptureComponent2D>(*It))
 			{
-				RemoveInstanceComponent(CaptureComponent);
-				CaptureComponent->UnregisterComponent();
-				CaptureComponent->DestroyComponent();
+				bool bDestroyComponent = false;
+
+				if (bHandleSceneCaptureRefactor && CaptureComponent->IsA<UCompositeShadowReflectionCatcherComponent>())
+				{
+					bDestroyComponent = true;
+				}
+
+				if (bHandleDerivedSceneCaptureTypes && (CaptureComponent->GetClass() == USceneCaptureComponent2D::StaticClass()))
+				{
+					bDestroyComponent = true;
+				}
+
+				if (bDestroyComponent)
+				{
+					RemoveInstanceComponent(CaptureComponent);
+					CaptureComponent->UnregisterComponent();
+					CaptureComponent->DestroyComponent();
+				}
 			}
 		}
 	}
@@ -567,7 +585,7 @@ void ACompositeActor::SetCamera(const FComponentReference& InComponentReference)
 }
 
 template<>
-USceneCaptureComponent2D* ACompositeActor::FindOrCreateSceneCapture(const UCompositeLayerBase* InLayer, int32 InIndex, FName InBaseName)
+UCompositeSceneCapture2DComponent* ACompositeActor::FindOrCreateSceneCapture(const UCompositeLayerBase* InLayer, int32 InIndex, FName InBaseName)
 {
 	if (!ensureMsgf(InLayer, TEXT("Expected non-null composite layer.")))
 	{
@@ -584,31 +602,20 @@ USceneCaptureComponent2D* ACompositeActor::FindOrCreateSceneCapture(const UCompo
 		SceneCapturesPerLayer.Add(InLayer, FSceneCaptureComponentArray{});
 	}
 
-	TArray<TObjectPtr<USceneCaptureComponent2D>>& SceneCaptures = SceneCapturesPerLayer[InLayer].Components;
+	TArray<TObjectPtr<UCompositeSceneCapture2DComponent>>& SceneCaptures = SceneCapturesPerLayer[InLayer].Components;
 
 	if (SceneCaptures.IsValidIndex(InIndex) && IsValid(SceneCaptures[InIndex]))
 	{
 		return SceneCaptures[InIndex];
 	}
 
-	const FName UniqueName = MakeUniqueObjectName(this, USceneCaptureComponent2D::StaticClass(), InBaseName);
-	USceneCaptureComponent2D* Component = NewObject<USceneCaptureComponent2D>(this, UniqueName, RF_Transactional);
+	const FName UniqueName = MakeUniqueObjectName(this, UCompositeSceneCapture2DComponent::StaticClass(), InBaseName);
+	UCompositeSceneCapture2DComponent* Component = NewObject<UCompositeSceneCapture2DComponent>(this, UniqueName, RF_Transactional);
 	Component->SetupAttachment(RootComponent);
 	Component->RegisterComponent();
 	Component->ComponentTags.Add(UE::Composite::DefaultCompositeTag);
 	AddInstanceComponent(Component);
 
-	UE::Composite::ConfigureSceneCapture(*Component, ESceneCapturePrimitiveRenderMode::PRM_UseShowOnlyList);
-	
-	TArray<FEngineShowFlagsSetting> ShowFlagSettings;
-	ShowFlagSettings.Add(FEngineShowFlagsSetting{ FEngineShowFlags::FindNameByIndex(FEngineShowFlags::EShowFlag::SF_Atmosphere), false });
-	ShowFlagSettings.Add(FEngineShowFlagsSetting{ FEngineShowFlags::FindNameByIndex(FEngineShowFlags::EShowFlag::SF_AllowPrimitiveAlphaHoldout), false });
-	ShowFlagSettings.Add(FEngineShowFlagsSetting{ FEngineShowFlags::FindNameByIndex(FEngineShowFlags::EShowFlag::SF_Bloom), false });
-	ShowFlagSettings.Add(FEngineShowFlagsSetting{ FEngineShowFlags::FindNameByIndex(FEngineShowFlags::EShowFlag::SF_EyeAdaptation), false });
-	ShowFlagSettings.Add(FEngineShowFlagsSetting{ FEngineShowFlags::FindNameByIndex(FEngineShowFlags::EShowFlag::SF_LocalExposure), false });
-	ShowFlagSettings.Add(FEngineShowFlagsSetting{ FEngineShowFlags::FindNameByIndex(FEngineShowFlags::EShowFlag::SF_PostProcessMaterial), false });
-	Component->SetShowFlagSettings(ShowFlagSettings);
-
 	SceneCaptures.EmplaceAt(InIndex, Component);
 
 #if WITH_EDITOR
@@ -636,7 +643,7 @@ UCompositeShadowReflectionCatcherComponent* ACompositeActor::FindOrCreateSceneCa
 		SceneCapturesPerLayer.Add(InLayer, {});
 	}
 
-	TArray<TObjectPtr<USceneCaptureComponent2D>>& SceneCaptureArray = SceneCapturesPerLayer[InLayer].Components;
+	TArray<TObjectPtr<UCompositeSceneCapture2DComponent>>& SceneCaptureArray = SceneCapturesPerLayer[InLayer].Components;
 
 	if (SceneCaptureArray.IsValidIndex(InIndex))
 	{
@@ -663,7 +670,7 @@ void ACompositeActor::ActivateSceneCaptures(const UCompositeLayerBase* InLayer,
 {
 	if (FSceneCaptureComponentArray* SceneCapturesPtr = SceneCapturesPerLayer.Find(InLayer))
 	{
-		for (TObjectPtr<USceneCaptureComponent2D>& SceneCapture : SceneCapturesPtr->Components)
+		for (TObjectPtr<UCompositeSceneCapture2DComponent>& SceneCapture : SceneCapturesPtr->Components)
 		{
 			if (IsValid(SceneCapture))
 			{
@@ -685,7 +692,7 @@ void ACompositeActor::DestroySceneCaptures(const UCompositeLayerBase* InLayer)
 	FSceneCaptureComponentArray RemovedArray;
 	if (SceneCapturesPerLayer.RemoveAndCopyValue(InLayer, RemovedArray))
 	{
-		for (TObjectPtr<USceneCaptureComponent2D>& RemovedComponent : RemovedArray.Components)
+		for (TObjectPtr<UCompositeSceneCapture2DComponent>& RemovedComponent : RemovedArray.Components)
 		{
 			if (IsValid(RemovedComponent))
 			{
diff --git a/Engine/Plugins/Experimental/Compositing/Composite/Source/Composite/Private/CompositeUtilities.cpp b/Engine/Plugins/Experimental/Compositing/Composite/Source/Composite/Private/CompositeUtilities.cpp
deleted file mode 100644
diff --git a/Engine/Plugins/Experimental/Compositing/Composite/Source/Composite/Private/CompositeUtilities.h b/Engine/Plugins/Experimental/Compositing/Composite/Source/Composite/Private/CompositeUtilities.h
deleted file mode 100644
diff --git a/Engine/Plugins/Experimental/Compositing/Composite/Source/Composite/Private/Layers/CompositeLayerSceneCapture.cpp b/Engine/Plugins/Experimental/Compositing/Composite/Source/Composite/Private/Layers/CompositeLayerSceneCapture.cpp
--- a/Engine/Plugins/Experimental/Compositing/Composite/Source/Composite/Private/Layers/CompositeLayerSceneCapture.cpp
+++ b/Engine/Plugins/Experimental/Compositing/Composite/Source/Composite/Private/Layers/CompositeLayerSceneCapture.cpp
@@ -4,9 +4,8 @@
 
 #include "CompositeActor.h"
 #include "Components/PrimitiveComponent.h"
-#include "Components/SceneCaptureComponent2D.h"
+#include "Components/CompositeSceneCapture2DComponent.h"
 #include "CompositeRenderTargetPool.h"
-#include "CompositeUtilities.h"
 #include "Engine/TextureRenderTarget2D.h"
 #include "UObject/Package.h"
 
@@ -33,11 +32,6 @@ void UCompositeLayerSceneCapture::SetEnabled(bool bInEnabled)
 void UCompositeLayerSceneCapture::BeginDestroy()
 {
 	Super::BeginDestroy();
-
-	if (USceneCaptureComponent2D* CaptureComponent = GetSceneCaptureComponent())
-	{
-		FCompositeRenderTargetPool::Get().ReleaseTarget(CaptureComponent->TextureTarget);
-	}
 }
 	
 #if WITH_EDITOR
@@ -154,7 +148,7 @@ void UCompositeLayerSceneCapture::SetCustomRenderPass(bool bInIsFastRenderPass)
 {
 	bCustomRenderPass = bInIsFastRenderPass;
 
-	USceneCaptureComponent2D* SceneCaptureComponent = GetSceneCaptureComponent();
+	UCompositeSceneCapture2DComponent* SceneCaptureComponent = GetSceneCaptureComponent();
 	if (IsValid(SceneCaptureComponent))
 	{
 		if (bCustomRenderPass)
@@ -170,12 +164,12 @@ void UCompositeLayerSceneCapture::SetCustomRenderPass(bool bInIsFastRenderPass)
 	}
 }
 
-USceneCaptureComponent2D* UCompositeLayerSceneCapture::GetSceneCaptureComponent() const
+UCompositeSceneCapture2DComponent* UCompositeLayerSceneCapture::GetSceneCaptureComponent() const
 {
 	ACompositeActor* CompositeActor = GetTypedOuter<ACompositeActor>();
 	if (IsValid(CompositeActor))
 	{
-		return CompositeActor->FindOrCreateSceneCapture<USceneCaptureComponent2D>(this);
+		return CompositeActor->FindOrCreateSceneCapture<UCompositeSceneCapture2DComponent>(this);
 	}
 
 	return nullptr;
diff --git a/Engine/Plugins/Experimental/Compositing/Composite/Source/Composite/Private/Layers/CompositeLayerSingleLightShadow.cpp b/Engine/Plugins/Experimental/Compositing/Composite/Source/Composite/Private/Layers/CompositeLayerSingleLightShadow.cpp
--- a/Engine/Plugins/Experimental/Compositing/Composite/Source/Composite/Private/Layers/CompositeLayerSingleLightShadow.cpp
+++ b/Engine/Plugins/Experimental/Compositing/Composite/Source/Composite/Private/Layers/CompositeLayerSingleLightShadow.cpp
@@ -6,7 +6,7 @@
 #include "CompositeActor.h"
 #include "CompositeRenderTargetPool.h"
 #include "Passes/CompositePassMaterial.h"
-#include "Components/SceneCaptureComponent2D.h"
+#include "Components/CompositeSceneCapture2DComponent.h"
 #include "Engine/DirectionalLight.h"
 #include "Engine/TextureRenderTarget2D.h"
 #include "Passes/CompositeCorePassProxy.h"
@@ -178,15 +178,6 @@ void UCompositeLayerSingleLightShadow::SetEnabled(bool bInEnabled)
 void UCompositeLayerSingleLightShadow::BeginDestroy()
 {
 	Super::BeginDestroy();
-
-	if (USceneCaptureComponent2D* CaptureComponent = CachedShadowMapCapture.Get())
-	{
-		FCompositeRenderTargetPool::Get().ReleaseTarget(CaptureComponent->TextureTarget);
-	}
-	if (USceneCaptureComponent2D* CaptureComponent = CachedSceneDepthCapture.Get())
-	{
-		FCompositeRenderTargetPool::Get().ReleaseTarget(CaptureComponent->TextureTarget);
-	}
 }
 
 #if WITH_EDITOR
@@ -340,9 +331,9 @@ void UCompositeLayerSingleLightShadow::SetShadowCastingActors(TArray<TSoftObject
 	CachedSceneDepthCapture->HiddenComponents = HiddenComponents;
 }
 
-USceneCaptureComponent2D* UCompositeLayerSingleLightShadow::FindOrCreateSceneDepthCapture(ACompositeActor& InOuter) const
+UCompositeSceneCapture2DComponent* UCompositeLayerSingleLightShadow::FindOrCreateSceneDepthCapture(ACompositeActor& InOuter) const
 {
-	USceneCaptureComponent2D* SceneDepthCapture = InOuter.FindOrCreateSceneCapture<USceneCaptureComponent2D>(this, 1, FName("SceneDepthCapture"));
+	UCompositeSceneCapture2DComponent* SceneDepthCapture = InOuter.FindOrCreateSceneCapture<UCompositeSceneCapture2DComponent>(this, 1, FName("SceneDepthCapture"));
 	
 	if (ensure(IsValid(SceneDepthCapture)))
 	{
@@ -355,9 +346,9 @@ USceneCaptureComponent2D* UCompositeLayerSingleLightShadow::FindOrCreateSceneDep
 	return SceneDepthCapture;
 }
 
-USceneCaptureComponent2D* UCompositeLayerSingleLightShadow::FindOrCreateShadowMapCapture(ACompositeActor& InOuter) const
+UCompositeSceneCapture2DComponent* UCompositeLayerSingleLightShadow::FindOrCreateShadowMapCapture(ACompositeActor& InOuter) const
 {
-	USceneCaptureComponent2D* ShadowMapCapture = InOuter.FindOrCreateSceneCapture<USceneCaptureComponent2D>(this, 0, FName("ShadowMapCapture"));
+	UCompositeSceneCapture2DComponent* ShadowMapCapture = InOuter.FindOrCreateSceneCapture<UCompositeSceneCapture2DComponent>(this, 0, FName("ShadowMapCapture"));
 
 	if (ensure(IsValid(ShadowMapCapture)))
 	{
@@ -375,7 +366,7 @@ USceneCaptureComponent2D* UCompositeLayerSingleLightShadow::FindOrCreateShadowMa
 	return ShadowMapCapture;
 }
 
-void UCompositeLayerSingleLightShadow::GetShadowMatrices(USceneCaptureComponent2D& ShadowMapCapture, FMatrix44f& OutShadowMatrix, FVector4f& OutShadowInvDeviceZToWorldZ) const
+void UCompositeLayerSingleLightShadow::GetShadowMatrices(UCompositeSceneCapture2DComponent& ShadowMapCapture, FMatrix44f& OutShadowMatrix, FVector4f& OutShadowInvDeviceZToWorldZ) const
 {
 	FMinimalViewInfo ViewInfo;
 	ShadowMapCapture.GetCameraView(GetWorld()->DeltaTimeSeconds, ViewInfo);
diff --git a/Engine/Plugins/Experimental/Compositing/Composite/Source/Composite/Public/Components/CompositeSceneCapture2DComponent.h b/Engine/Plugins/Experimental/Compositing/Composite/Source/Composite/Public/Components/CompositeSceneCapture2DComponent.h
new file mode 100644
--- /dev/null
+++ b/Engine/Plugins/Experimental/Compositing/Composite/Source/Composite/Public/Components/CompositeSceneCapture2DComponent.h
@@ -0,0 +1,33 @@
+// Copyright Epic Games, Inc. All Rights Reserved.
+
+#pragma once
+
+#include "CoreTypes.h"
+#include "Components/SceneCaptureComponent2D.h"
+
+#include "CompositeSceneCapture2DComponent.generated.h"
+
+#define UE_API COMPOSITE_API
+
+UCLASS(MinimalAPI, ClassGroup = Composite, EditInlineNew, meta = (BlueprintSpawnableComponent))
+class UCompositeSceneCapture2DComponent : public USceneCaptureComponent2D
+{
+	GENERATED_UCLASS_BODY()
+
+public:
+
+	//~ Begin UObject Interface
+	UE_API virtual void BeginDestroy() override;
+
+#if WITH_EDITOR
+	UE_API virtual void PostEditChangeProperty(struct FPropertyChangedEvent& PropertyChangedEvent) override;
+#endif //WITH_EDITOR
+	//~ End UObject Interface
+
+	//~ Begin USceneCaptureComponent Interface
+	UE_API virtual const AActor* GetViewOwner() const override;
+	//~ End USceneCaptureComponent Interface
+};
+
+#undef UE_API
+
diff --git a/Engine/Plugins/Experimental/Compositing/Composite/Source/Composite/Public/Components/CompositeShadowReflectionCatcherComponent.h b/Engine/Plugins/Experimental/Compositing/Composite/Source/Composite/Public/Components/CompositeShadowReflectionCatcherComponent.h
--- a/Engine/Plugins/Experimental/Compositing/Composite/Source/Composite/Public/Components/CompositeShadowReflectionCatcherComponent.h
+++ b/Engine/Plugins/Experimental/Compositing/Composite/Source/Composite/Public/Components/CompositeShadowReflectionCatcherComponent.h
@@ -3,26 +3,17 @@
 #pragma once
 
 #include "CoreTypes.h"
-#include "Components/SceneCaptureComponent2D.h"
+#include "CompositeSceneCapture2DComponent.h"
 
 #include "CompositeShadowReflectionCatcherComponent.generated.h"
 
 #define UE_API COMPOSITE_API
 
 UCLASS(MinimalAPI, ClassGroup = Composite, EditInlineNew, meta = (BlueprintSpawnableComponent))
-class UCompositeShadowReflectionCatcherComponent : public USceneCaptureComponent2D
+class UCompositeShadowReflectionCatcherComponent : public UCompositeSceneCapture2DComponent
 {
 	GENERATED_UCLASS_BODY()
 
-public:
-
-	//~ Begin UObject Interface
-	UE_API virtual void BeginDestroy() override;
-	//~ End UObject Interface
-
-#if WITH_EDITOR
-	UE_API virtual void PostEditChangeProperty(struct FPropertyChangedEvent& PropertyChangedEvent) override;
-#endif //WITH_EDITOR
 };
 
 #undef UE_API
diff --git a/Engine/Plugins/Experimental/Compositing/Composite/Source/Composite/Public/CompositeActor.h b/Engine/Plugins/Experimental/Compositing/Composite/Source/Composite/Public/CompositeActor.h
--- a/Engine/Plugins/Experimental/Compositing/Composite/Source/Composite/Public/CompositeActor.h
+++ b/Engine/Plugins/Experimental/Compositing/Composite/Source/Composite/Public/CompositeActor.h
@@ -16,7 +16,7 @@ class ACameraActor;
 class UCompositeCoreSubsystem;
 class UCompositeShadowReflectionCatcherComponent;
 class UPrimitiveComponent;
-class USceneCaptureComponent2D;
+class UCompositeSceneCapture2DComponent;
 
 DECLARE_DYNAMIC_MULTICAST_DELEGATE(FOnCompositeActorJoinMultiUserSession);
 DECLARE_DELEGATE_TwoParams(FOnCompositeActorEnabledStateChange, bool /* */, UCompositeCoreSubsystem* /* */);
@@ -27,7 +27,7 @@ struct FSceneCaptureComponentArray
 	GENERATED_BODY()
 
 	UPROPERTY()
-	TArray<TObjectPtr<USceneCaptureComponent2D>> Components;
+	TArray<TObjectPtr<UCompositeSceneCapture2DComponent>> Components;
 };
 
 /* Main render color space & encoding output mode. */
diff --git a/Engine/Plugins/Experimental/Compositing/Composite/Source/Composite/Public/Layers/CompositeLayerSceneCapture.h b/Engine/Plugins/Experimental/Compositing/Composite/Source/Composite/Public/Layers/CompositeLayerSceneCapture.h
--- a/Engine/Plugins/Experimental/Compositing/Composite/Source/Composite/Public/Layers/CompositeLayerSceneCapture.h
+++ b/Engine/Plugins/Experimental/Compositing/Composite/Source/Composite/Public/Layers/CompositeLayerSceneCapture.h
@@ -59,7 +59,7 @@ public:
 
 private:
 	/** Convenience function to access the actor-managed scene capture component. */
-	class USceneCaptureComponent2D* GetSceneCaptureComponent() const;
+	class UCompositeSceneCapture2DComponent* GetSceneCaptureComponent() const;
 
 	/**
 	 * List of actors to be rendered from the scene capture.
diff --git a/Engine/Plugins/Experimental/Compositing/Composite/Source/Composite/Public/Layers/CompositeLayerSingleLightShadow.h b/Engine/Plugins/Experimental/Compositing/Composite/Source/Composite/Public/Layers/CompositeLayerSingleLightShadow.h
--- a/Engine/Plugins/Experimental/Compositing/Composite/Source/Composite/Public/Layers/CompositeLayerSingleLightShadow.h
+++ b/Engine/Plugins/Experimental/Compositing/Composite/Source/Composite/Public/Layers/CompositeLayerSingleLightShadow.h
@@ -17,7 +17,7 @@
 
 class AActor;
 class ALight;
-class USceneCaptureComponent2D;
+class UCompositeSceneCapture2DComponent;
 
 /**
 * Layer for catching shadows from a single light.
@@ -94,13 +94,13 @@ public:
 
 private:
 	/** Find or create scene depth capture component, correctly configured. */
-	USceneCaptureComponent2D* FindOrCreateSceneDepthCapture(ACompositeActor& InOuter) const;
+	UCompositeSceneCapture2DComponent* FindOrCreateSceneDepthCapture(ACompositeActor& InOuter) const;
 
 	/** Find or create shadow map capture component, correctly configured. */
-	USceneCaptureComponent2D* FindOrCreateShadowMapCapture(ACompositeActor& InOuter) const;
+	UCompositeSceneCapture2DComponent* FindOrCreateShadowMapCapture(ACompositeActor& InOuter) const;
 
 	/** Convenience function to calculate shadow matrices. */
-	void GetShadowMatrices(USceneCaptureComponent2D& ShadowMapCapture, FMatrix44f& OutShadowMatrix, FVector4f& OutShadowInvDeviceZToWorldZ) const;
+	void GetShadowMatrices(UCompositeSceneCapture2DComponent& ShadowMapCapture, FMatrix44f& OutShadowMatrix, FVector4f& OutShadowInvDeviceZToWorldZ) const;
 
 private:
 	/**
@@ -111,10 +111,10 @@ private:
 	TArray<TSoftObjectPtr<AActor>> ShadowCastingActors;
 
 	/** Cached pointer to the composite-managed shadow map scene capture. */
-	mutable TWeakObjectPtr<USceneCaptureComponent2D> CachedShadowMapCapture;
+	mutable TWeakObjectPtr<UCompositeSceneCapture2DComponent> CachedShadowMapCapture;
 
 	/** Cached pointer to the scene depth capture. */
-	mutable TWeakObjectPtr<USceneCaptureComponent2D> CachedSceneDepthCapture;
+	mutable TWeakObjectPtr<UCompositeSceneCapture2DComponent> CachedSceneDepthCapture;
 };
 
 #undef UE_API
diff --git a/Engine/Source/Runtime/Core/Public/UObject/UE5ReleaseStreamObjectVersions.inl b/Engine/Source/Runtime/Core/Public/UObject/UE5ReleaseStreamObjectVersions.inl
--- a/Engine/Source/Runtime/Core/Public/UObject/UE5ReleaseStreamObjectVersions.inl
+++ b/Engine/Source/Runtime/Core/Public/UObject/UE5ReleaseStreamObjectVersions.inl
@@ -177,5 +177,8 @@ UE5_RELEASE_VERSION(WorldPartitionAddStandaloneHLODSupport, 55)
 // Fixed the missing bounds for cloth assets that don't have them serialized
 UE5_RELEASE_VERSION(RecalculateClothAssetSerializedBounds, 56)
 
+// Composite plugin now uses its own derived scene capture components
+UE5_RELEASE_VERSION(CompositePluginDerivedSceneCaptures, 57)
+
 // -----<new versions can be added above this line>-------------------------------------------------
 #undef UE5_RELEASE_VERSION
\ No newline at end of file
