Composite: Fix crash when adding an actor before a light on the single light layer.

diff --git a/Engine/Plugins/Experimental/Compositing/Composite/Source/Composite/Private/CompositeActor.cpp b/Engine/Plugins/Experimental/Compositing/Composite/Source/Composite/Private/CompositeActor.cpp
--- a/Engine/Plugins/Experimental/Compositing/Composite/Source/Composite/Private/CompositeActor.cpp
+++ b/Engine/Plugins/Experimental/Compositing/Composite/Source/Composite/Private/CompositeActor.cpp
@@ -615,8 +615,13 @@ UCompositeSceneCapture2DComponent* ACompositeActor::FindOrCreateSceneCapture(con
 	Component->RegisterComponent();
 	Component->ComponentTags.Add(UE::Composite::DefaultCompositeTag);
 	AddInstanceComponent(Component);
-
-	SceneCaptures.EmplaceAt(InIndex, Component);
+	
+	const int32 MinSize = InIndex + 1;
+	if (SceneCaptures.Num() < MinSize)
+	{
+		SceneCaptures.SetNumZeroed(MinSize);
+	}
+	SceneCaptures[InIndex] = Component;
 
 #if WITH_EDITOR
 	FModuleManager::LoadModuleChecked<FLevelEditorModule>("LevelEditor").BroadcastComponentsEdited();
