Composure: Added custom actor picker to scene capture and single light shadow layers

diff --git a/Engine/Plugins/Experimental/Compositing/Composite/Source/Composite/Public/Layers/CompositeLayerSceneCapture.h b/Engine/Plugins/Experimental/Compositing/Composite/Source/Composite/Public/Layers/CompositeLayerSceneCapture.h
--- a/Engine/Plugins/Experimental/Compositing/Composite/Source/Composite/Public/Layers/CompositeLayerSceneCapture.h
+++ b/Engine/Plugins/Experimental/Compositing/Composite/Source/Composite/Public/Layers/CompositeLayerSceneCapture.h
@@ -76,6 +76,8 @@ private:
 	 */
 	UPROPERTY(EditAnywhere, BlueprintGetter = IsCustomRenderPass, BlueprintSetter = SetCustomRenderPass, Category = "", meta = (AllowPrivateAccess = true))
 	bool bCustomRenderPass;
+
+	friend class FCompositeLayerSceneCaptureCustomization;
 };
 
 #undef UE_API
diff --git a/Engine/Plugins/Experimental/Compositing/Composite/Source/Composite/Public/Layers/CompositeLayerSingleLightShadow.h b/Engine/Plugins/Experimental/Compositing/Composite/Source/Composite/Public/Layers/CompositeLayerSingleLightShadow.h
--- a/Engine/Plugins/Experimental/Compositing/Composite/Source/Composite/Public/Layers/CompositeLayerSingleLightShadow.h
+++ b/Engine/Plugins/Experimental/Compositing/Composite/Source/Composite/Public/Layers/CompositeLayerSingleLightShadow.h
@@ -115,6 +115,8 @@ private:
 
 	/** Cached pointer to the scene depth capture. */
 	mutable TWeakObjectPtr<UCompositeSceneCapture2DComponent> CachedSceneDepthCapture;
+
+	friend class FCompositeLayerSingleLightShadowCustomization;
 };
 
 #undef UE_API
diff --git a/Engine/Plugins/Experimental/Compositing/Composite/Source/CompositeEditor/Private/CompositeEditorModule.cpp b/Engine/Plugins/Experimental/Compositing/Composite/Source/CompositeEditor/Private/CompositeEditorModule.cpp
--- a/Engine/Plugins/Experimental/Compositing/Composite/Source/CompositeEditor/Private/CompositeEditorModule.cpp
+++ b/Engine/Plugins/Experimental/Compositing/Composite/Source/CompositeEditor/Private/CompositeEditorModule.cpp
@@ -17,9 +17,13 @@
 #include "PropertyEditorModule.h"
 #include "Customizations/CompositeActorCustomization.h"
 #include "Customizations/CompositeLayerPlateCustomization.h"
+#include "Customizations/CompositeLayerSceneCaptureCustomization.h"
 #include "Customizations/CompositeLayerShadowReflectionCustomization.h"
+#include "Customizations/CompositeLayerSingleLightShadowCustomization.h"
 #include "Layers/CompositeLayerPlate.h"
+#include "Layers/CompositeLayerSceneCapture.h"
 #include "Layers/CompositeLayerShadowReflection.h"
+#include "Layers/CompositeLayerSingleLightShadow.h"
 #include "UI/SCompositeEditorPanel.h"
 
 #define LOCTEXT_NAMESPACE "CompositeEditorModule"
@@ -157,6 +161,16 @@ void FCompositeEditorModule::RegisterCustomizations()
 		FOnGetDetailCustomizationInstance::CreateStatic(&FCompositeLayerShadowReflectionCustomization::MakeInstance)
 	);
 
+	PropertyModule.RegisterCustomClassLayout(
+			UCompositeLayerSceneCapture::StaticClass()->GetFName(),
+			FOnGetDetailCustomizationInstance::CreateStatic(&FCompositeLayerSceneCaptureCustomization::MakeInstance)
+		);
+
+	PropertyModule.RegisterCustomClassLayout(
+			UCompositeLayerSingleLightShadow::StaticClass()->GetFName(),
+			FOnGetDetailCustomizationInstance::CreateStatic(&FCompositeLayerSingleLightShadowCustomization::MakeInstance)
+		);
+
 	PropertyModule.RegisterCustomClassLayout(
 		ACompositeActor::StaticClass()->GetFName(),
 		FOnGetDetailCustomizationInstance::CreateStatic(&FCompositeActorCustomization::MakeInstance)
@@ -172,6 +186,8 @@ void FCompositeEditorModule::UnregisterCustomizations()
 			PropertyModule->UnregisterCustomClassLayout(ACompositeActor::StaticClass()->GetFName());
 			PropertyModule->UnregisterCustomClassLayout(UCompositeLayerPlate::StaticClass()->GetFName());
 			PropertyModule->UnregisterCustomClassLayout(UCompositeLayerShadowReflection::StaticClass()->GetFName());
+			PropertyModule->UnregisterCustomClassLayout(UCompositeLayerSceneCapture::StaticClass()->GetFName());
+			PropertyModule->UnregisterCustomClassLayout(UCompositeLayerSingleLightShadow::StaticClass()->GetFName());
 		}
 	}
 }
diff --git a/Engine/Plugins/Experimental/Compositing/Composite/Source/CompositeEditor/Private/Customizations/CompositeLayerSceneCaptureCustomization.cpp b/Engine/Plugins/Experimental/Compositing/Composite/Source/CompositeEditor/Private/Customizations/CompositeLayerSceneCaptureCustomization.cpp
new file mode 100644
--- /dev/null
+++ b/Engine/Plugins/Experimental/Compositing/Composite/Source/CompositeEditor/Private/Customizations/CompositeLayerSceneCaptureCustomization.cpp
@@ -0,0 +1,88 @@
+﻿// Copyright Epic Games, Inc. All Rights Reserved.
+
+#include "CompositeLayerSceneCaptureCustomization.h"
+
+#include "DetailCategoryBuilder.h"
+#include "DetailLayoutBuilder.h"
+#include "DetailWidgetRow.h"
+#include "IDetailGroup.h"
+#include "IPropertyUtilities.h"
+#include "Layers/CompositeLayerSceneCapture.h"
+#include "UI/SCompositeActorPickerTable.h"
+#include "Widgets/Text/STextBlock.h"
+
+#define LOCTEXT_NAMESPACE "FCompositeLayerSceneCaptureCustomization"
+
+TSharedRef<IDetailCustomization> FCompositeLayerSceneCaptureCustomization::MakeInstance()
+{
+	return MakeShared<FCompositeLayerSceneCaptureCustomization>();
+}
+
+void FCompositeLayerSceneCaptureCustomization::CustomizeDetails(const TSharedPtr<IDetailLayoutBuilder>& DetailBuilder)
+{
+	CachedDetailBuilder = DetailBuilder;
+	IDetailCustomization::CustomizeDetails(DetailBuilder);
+}
+
+void FCompositeLayerSceneCaptureCustomization::CustomizeDetails(IDetailLayoutBuilder& DetailLayout)
+{
+	DetailLayout.HideProperty(GET_MEMBER_NAME_CHECKED(UCompositeLayerSceneCapture, Actors));
+	
+	TArray<TWeakObjectPtr<UCompositeLayerSceneCapture>> Objects = DetailLayout.GetObjectsOfTypeBeingCustomized<UCompositeLayerSceneCapture>();
+
+	IDetailCategoryBuilder& LayerCategory = DetailLayout.EditCategoryAllowNone(NAME_None);
+	
+	// Add all default simple properties ahead of the custom group
+	TArray<TSharedRef<IPropertyHandle>> LayerProperties;
+	LayerCategory.GetDefaultProperties(LayerProperties, /* bSimpleProperties */ true, /* bAdvancedProperties */ false);
+	for (const TSharedRef<IPropertyHandle>& Property : LayerProperties)
+	{
+		if (!Property->IsCustomized())
+		{
+			LayerCategory.AddProperty(Property);
+		}
+	}
+
+	if (Objects.Num() == 1 && Objects[0].IsValid())
+	{
+		UCompositeLayerSceneCapture* SceneCapture = Cast<UCompositeLayerSceneCapture>(Objects[0].Get());
+		
+		IDetailGroup& ActorListGroup = LayerCategory.AddGroup("ShadowReflectionCatcherContent", LOCTEXT("ActorListGroupName", "Scene Capture Content"), false, true);
+		FCompositeActorPickerListRef ActorListRef(SceneCapture, GET_MEMBER_NAME_CHECKED(UCompositeLayerSceneCapture, Actors), &SceneCapture->Actors);
+		
+		ActorListGroup.AddWidgetRow()
+		.WholeRowContent()
+		[
+			SNew(SCompositeActorPickerTable, ActorListRef)
+			.OnLayoutSizeChanged(this, &FCompositeLayerSceneCaptureCustomization::OnLayoutSizeChanged)
+		];
+	}
+	else
+	{
+		// Can't display actor list if multiple layers are selected, so simply put a "Multiple Values" entry in the property list
+		LayerCategory.AddCustomRow(LOCTEXT("ActorListGroupName", "Scene Capture Content"))
+		.NameContent()
+		[
+			SNew(STextBlock)
+			.Text(LOCTEXT("ActorListGroupName", "Scene Capture Content"))
+			.Font(IDetailLayoutBuilder::GetDetailFont())
+		]
+		.ValueContent()
+		[
+			SNew(STextBlock)
+			.Text(LOCTEXT("MultipleValues", "Multiple Values"))
+			.Font(IDetailLayoutBuilder::GetDetailFont())
+		];
+	}
+}
+
+void FCompositeLayerSceneCaptureCustomization::OnLayoutSizeChanged()
+{
+	TSharedPtr<IDetailLayoutBuilder> PinnedDetailBuilder = CachedDetailBuilder.Pin();
+	if (PinnedDetailBuilder.IsValid())
+	{
+		PinnedDetailBuilder->GetPropertyUtilities()->RequestRefresh();
+	}
+}
+
+#undef LOCTEXT_NAMESPACE
diff --git a/Engine/Plugins/Experimental/Compositing/Composite/Source/CompositeEditor/Private/Customizations/CompositeLayerSceneCaptureCustomization.h b/Engine/Plugins/Experimental/Compositing/Composite/Source/CompositeEditor/Private/Customizations/CompositeLayerSceneCaptureCustomization.h
new file mode 100644
--- /dev/null
+++ b/Engine/Plugins/Experimental/Compositing/Composite/Source/CompositeEditor/Private/Customizations/CompositeLayerSceneCaptureCustomization.h
@@ -0,0 +1,27 @@
+﻿// Copyright Epic Games, Inc. All Rights Reserved.
+
+#pragma once
+
+#include "IDetailCustomization.h"
+
+/**
+ * Customization for the composite scene capture layer, primary for displaying a custom widget for picking actors
+ */
+class FCompositeLayerSceneCaptureCustomization : public IDetailCustomization
+{
+public:
+	/** Makes a new instance of the details customization */
+	static TSharedRef<IDetailCustomization> MakeInstance();
+
+	// IDetailCustomization interface
+	virtual void CustomizeDetails(const TSharedPtr<IDetailLayoutBuilder>& DetailBuilder) override;
+	virtual void CustomizeDetails(IDetailLayoutBuilder& DetailLayout) override;
+	// End of IDetailCustomization interface
+
+private:
+	/** Raised when the custom widgets' layout size may have changed */
+	void OnLayoutSizeChanged();
+	
+private:
+	TWeakPtr<IDetailLayoutBuilder> CachedDetailBuilder;
+};
diff --git a/Engine/Plugins/Experimental/Compositing/Composite/Source/CompositeEditor/Private/Customizations/CompositeLayerSingleLightShadowCustomization.cpp b/Engine/Plugins/Experimental/Compositing/Composite/Source/CompositeEditor/Private/Customizations/CompositeLayerSingleLightShadowCustomization.cpp
new file mode 100644
--- /dev/null
+++ b/Engine/Plugins/Experimental/Compositing/Composite/Source/CompositeEditor/Private/Customizations/CompositeLayerSingleLightShadowCustomization.cpp
@@ -0,0 +1,88 @@
+﻿// Copyright Epic Games, Inc. All Rights Reserved.
+
+#include "CompositeLayerSingleLightShadowCustomization.h"
+
+#include "DetailCategoryBuilder.h"
+#include "DetailLayoutBuilder.h"
+#include "DetailWidgetRow.h"
+#include "IDetailGroup.h"
+#include "IPropertyUtilities.h"
+#include "Layers/CompositeLayerSingleLightShadow.h"
+#include "UI/SCompositeActorPickerTable.h"
+#include "Widgets/Text/STextBlock.h"
+
+#define LOCTEXT_NAMESPACE "FCompositeLayerSceneCaptureCustomization"
+
+TSharedRef<IDetailCustomization> FCompositeLayerSingleLightShadowCustomization::MakeInstance()
+{
+	return MakeShared<FCompositeLayerSingleLightShadowCustomization>();
+}
+
+void FCompositeLayerSingleLightShadowCustomization::CustomizeDetails(const TSharedPtr<IDetailLayoutBuilder>& DetailBuilder)
+{
+	CachedDetailBuilder = DetailBuilder;
+	IDetailCustomization::CustomizeDetails(DetailBuilder);
+}
+
+void FCompositeLayerSingleLightShadowCustomization::CustomizeDetails(IDetailLayoutBuilder& DetailLayout)
+{
+	DetailLayout.HideProperty(GET_MEMBER_NAME_CHECKED(UCompositeLayerSingleLightShadow, ShadowCastingActors));
+	
+	TArray<TWeakObjectPtr<UCompositeLayerSingleLightShadow>> Objects = DetailLayout.GetObjectsOfTypeBeingCustomized<UCompositeLayerSingleLightShadow>();
+
+	IDetailCategoryBuilder& LayerCategory = DetailLayout.EditCategoryAllowNone(NAME_None);
+	
+	// Add all default simple properties ahead of the custom group
+	TArray<TSharedRef<IPropertyHandle>> LayerProperties;
+	LayerCategory.GetDefaultProperties(LayerProperties, /* bSimpleProperties */ true, /* bAdvancedProperties */ false);
+	for (const TSharedRef<IPropertyHandle>& Property : LayerProperties)
+	{
+		if (!Property->IsCustomized())
+		{
+			LayerCategory.AddProperty(Property);
+		}
+	}
+
+	if (Objects.Num() == 1 && Objects[0].IsValid())
+	{
+		UCompositeLayerSingleLightShadow* SingleLightShadow = Cast<UCompositeLayerSingleLightShadow>(Objects[0].Get());
+		
+		IDetailGroup& ActorListGroup = LayerCategory.AddGroup("ShadowReflectionCatcherContent", LOCTEXT("ActorListGroupName", "Shadow Casting Content"), false, true);
+		FCompositeActorPickerListRef ActorListRef(SingleLightShadow, GET_MEMBER_NAME_CHECKED(UCompositeLayerSingleLightShadow, ShadowCastingActors), &SingleLightShadow->ShadowCastingActors);
+		
+		ActorListGroup.AddWidgetRow()
+		.WholeRowContent()
+		[
+			SNew(SCompositeActorPickerTable, ActorListRef)
+			.OnLayoutSizeChanged(this, &FCompositeLayerSingleLightShadowCustomization::OnLayoutSizeChanged)
+		];
+	}
+	else
+	{
+		// Can't display actor list if multiple layers are selected, so simply put a "Multiple Values" entry in the property list
+		LayerCategory.AddCustomRow(LOCTEXT("ActorListGroupName", "Shadow Casting Content"))
+		.NameContent()
+		[
+			SNew(STextBlock)
+			.Text(LOCTEXT("ActorListGroupName", "Shadow Casting Content"))
+			.Font(IDetailLayoutBuilder::GetDetailFont())
+		]
+		.ValueContent()
+		[
+			SNew(STextBlock)
+			.Text(LOCTEXT("MultipleValues", "Multiple Values"))
+			.Font(IDetailLayoutBuilder::GetDetailFont())
+		];
+	}
+}
+
+void FCompositeLayerSingleLightShadowCustomization::OnLayoutSizeChanged()
+{
+	TSharedPtr<IDetailLayoutBuilder> PinnedDetailBuilder = CachedDetailBuilder.Pin();
+	if (PinnedDetailBuilder.IsValid())
+	{
+		PinnedDetailBuilder->GetPropertyUtilities()->RequestRefresh();
+	}
+}
+
+#undef LOCTEXT_NAMESPACE
diff --git a/Engine/Plugins/Experimental/Compositing/Composite/Source/CompositeEditor/Private/Customizations/CompositeLayerSingleLightShadowCustomization.h b/Engine/Plugins/Experimental/Compositing/Composite/Source/CompositeEditor/Private/Customizations/CompositeLayerSingleLightShadowCustomization.h
new file mode 100644
--- /dev/null
+++ b/Engine/Plugins/Experimental/Compositing/Composite/Source/CompositeEditor/Private/Customizations/CompositeLayerSingleLightShadowCustomization.h
@@ -0,0 +1,27 @@
+﻿// Copyright Epic Games, Inc. All Rights Reserved.
+
+#pragma once
+
+#include "IDetailCustomization.h"
+
+/**
+ * Customization for the composite single light shadow layer, primary for displaying a custom widget for picking actors
+ */
+class FCompositeLayerSingleLightShadowCustomization : public IDetailCustomization
+{
+public:
+	/** Makes a new instance of the details customization */
+	static TSharedRef<IDetailCustomization> MakeInstance();
+
+	// IDetailCustomization interface
+	virtual void CustomizeDetails(const TSharedPtr<IDetailLayoutBuilder>& DetailBuilder) override;
+	virtual void CustomizeDetails(IDetailLayoutBuilder& DetailLayout) override;
+	// End of IDetailCustomization interface
+
+private:
+	/** Raised when the custom widgets' layout size may have changed */
+	void OnLayoutSizeChanged();
+	
+private:
+	TWeakPtr<IDetailLayoutBuilder> CachedDetailBuilder;
+};
