Composure: Fixed LOCTEXT collisions with UE-323188

diff --git a/Engine/Plugins/Experimental/Compositing/Composite/Source/CompositeEditor/Private/Customizations/CompositeLayerSceneCaptureCustomization.cpp b/Engine/Plugins/Experimental/Compositing/Composite/Source/CompositeEditor/Private/Customizations/CompositeLayerSceneCaptureCustomization.cpp
--- a/Engine/Plugins/Experimental/Compositing/Composite/Source/CompositeEditor/Private/Customizations/CompositeLayerSceneCaptureCustomization.cpp
+++ b/Engine/Plugins/Experimental/Compositing/Composite/Source/CompositeEditor/Private/Customizations/CompositeLayerSceneCaptureCustomization.cpp
@@ -47,7 +47,7 @@ void FCompositeLayerSceneCaptureCustomization::CustomizeDetails(IDetailLayoutBui
 	{
 		UCompositeLayerSceneCapture* SceneCapture = Cast<UCompositeLayerSceneCapture>(Objects[0].Get());
 		
-		IDetailGroup& ActorListGroup = LayerCategory.AddGroup("ShadowReflectionCatcherContent", LOCTEXT("ActorListGroupName", "Scene Capture Content"), false, true);
+		IDetailGroup& ActorListGroup = LayerCategory.AddGroup("ActorContent", LOCTEXT("ActorListGroupName", "Scene Capture Content"), false, true);
 		FCompositeActorPickerListRef ActorListRef(SceneCapture, GET_MEMBER_NAME_CHECKED(UCompositeLayerSceneCapture, Actors), &SceneCapture->Actors);
 		
 		ActorListGroup.AddWidgetRow()
diff --git a/Engine/Plugins/Experimental/Compositing/Composite/Source/CompositeEditor/Private/Customizations/CompositeLayerSingleLightShadowCustomization.cpp b/Engine/Plugins/Experimental/Compositing/Composite/Source/CompositeEditor/Private/Customizations/CompositeLayerSingleLightShadowCustomization.cpp
--- a/Engine/Plugins/Experimental/Compositing/Composite/Source/CompositeEditor/Private/Customizations/CompositeLayerSingleLightShadowCustomization.cpp
+++ b/Engine/Plugins/Experimental/Compositing/Composite/Source/CompositeEditor/Private/Customizations/CompositeLayerSingleLightShadowCustomization.cpp
@@ -11,7 +11,7 @@
 #include "UI/SCompositeActorPickerTable.h"
 #include "Widgets/Text/STextBlock.h"
 
-#define LOCTEXT_NAMESPACE "FCompositeLayerSceneCaptureCustomization"
+#define LOCTEXT_NAMESPACE "FCompositeLayerSingleLightShadowCustomization"
 
 TSharedRef<IDetailCustomization> FCompositeLayerSingleLightShadowCustomization::MakeInstance()
 {
@@ -47,7 +47,7 @@ void FCompositeLayerSingleLightShadowCustomization::CustomizeDetails(IDetailLayo
 	{
 		UCompositeLayerSingleLightShadow* SingleLightShadow = Cast<UCompositeLayerSingleLightShadow>(Objects[0].Get());
 		
-		IDetailGroup& ActorListGroup = LayerCategory.AddGroup("ShadowReflectionCatcherContent", LOCTEXT("ActorListGroupName", "Shadow Casting Content"), false, true);
+		IDetailGroup& ActorListGroup = LayerCategory.AddGroup("ActorContent", LOCTEXT("ActorListGroupName", "Shadow Casting Content"), false, true);
 		FCompositeActorPickerListRef ActorListRef(SingleLightShadow, GET_MEMBER_NAME_CHECKED(UCompositeLayerSingleLightShadow, ShadowCastingActors), &SingleLightShadow->ShadowCastingActors);
 		
 		ActorListGroup.AddWidgetRow()
