Composite: Rename delta keyer to more generic color keyer.

#rnx

diff --git a/Engine/Plugins/Experimental/Compositing/Composite/Config/BaseComposite.ini b/Engine/Plugins/Experimental/Compositing/Composite/Config/BaseComposite.ini
--- a/Engine/Plugins/Experimental/Compositing/Composite/Config/BaseComposite.ini
+++ b/Engine/Plugins/Experimental/Compositing/Composite/Config/BaseComposite.ini
@@ -33,4 +33,7 @@
 +PropertyRedirects=(OldName="/Script/Composite.CompositeLayerPlate.Postprocessing", NewName="/Script/Composite.CompositeLayerPlate.LayerPasses")
 +PropertyRedirects=(OldName="/Script/Composite.CompositeLayerSceneCapture.Postprocessing", NewName="/Script/Composite.CompositeLayerSceneCapture.LayerPasses")
 
++EnumRedirects=(OldName="/Script/Composite.CompositeDeltaKeyerVisualization",NewName="/Script/Composite.CompositeColorKeyerVisualization")
++EnumRedirects=(OldName="/Script/Composite.CompositeDeltaKeyerScreenType",NewName="/Script/Composite.CompositeColorKeyerScreenType")
++ClassRedirects=(OldName="/Script/Composite.CompositePassDeltaKeyer", NewName="/Script/Composite.CompositePassColorKeyer")
 
diff --git a/Engine/Plugins/Experimental/Compositing/Composite/Shaders/Private/CompositeColorKeyer.usf b/Engine/Plugins/Experimental/Compositing/Composite/Shaders/Private/CompositeColorKeyer.usf
new file mode 100644
--- /dev/null
+++ b/Engine/Plugins/Experimental/Compositing/Composite/Shaders/Private/CompositeColorKeyer.usf
@@ -0,0 +1,139 @@
+// Copyright Epic Games, Inc. All Rights Reserved.
+
+#include "/Engine/Private/Common.ush"
+#include "/Engine/Private/ScreenPass.ush"
+
+#include "CompositeColorKeyerCommon.ush"
+
+#ifndef TILE_SIZE
+    #error TILE_SIZE must be defined
+#endif
+
+Texture2D InputTexture;
+SamplerState InputSampler;
+
+Texture2D CleanPlateTexture;
+SamplerState CleanPlateSampler;
+
+RWTexture2D<float4> RWOutputTexture;
+
+SCREEN_PASS_TEXTURE_VIEWPORT(Input)
+SCREEN_PASS_TEXTURE_VIEWPORT(Output)
+
+float3 KeyColor;
+float4 Params0;
+float4 Params1;
+uint ScreenType;
+uint Visualization;
+uint bInvertAlpha;
+uint bPreserveVignetteAfterKey;
+
+float2 UVToScreenPos(float2 UV, float2 Extent)
+{
+	return (UV * Extent - Input_ScreenPosToViewportBias) / Input_ScreenPosToViewportScale;
+}
+
+// Output color function with debug convenience
+float4 Output(float4 Color)
+{
+	// Optionally invert alpha
+	Color.a = select(bInvertAlpha > 0, 1.0f - Color.a, Color.a);
+
+	if(Visualization == 1)
+	{
+		// Opaque to show fill
+		Color.a = 1.0f;
+	}
+	else if(Visualization == 2)
+	{
+		// Optionally show the alpha key matte
+		Color.rgb = sRGBToLinear(Color.a);
+		Color.a = 1.0f;
+	}
+
+	return Color;
+}
+
+[numthreads(TILE_SIZE, TILE_SIZE, 1)]
+void MainCS(
+    uint2 GroupId : SV_GroupID,
+    uint2 GroupThreadId : SV_GroupThreadID,
+    uint GroupThreadIndex : SV_GroupIndex,
+    uint2 DispatchThreadId : SV_DispatchThreadID)
+{
+	float4 SvPosition = float4((float2)DispatchThreadId + Output_ViewportMin + 0.5f, 0.0f, 1.0f);
+	float2 UV = (SvPosition.xy) * Output_ExtentInverse;
+
+	if (IsComputeUVOutOfBounds(UV))
+	{
+		return;
+	}
+
+	// Input parameters
+	const float3 ChannelWeights = Params0.xyz;
+	const float DespillStrength = Params0.w;
+	const float AlphaMin = Params1.x;
+	const float AlphaMax = Params1.y;
+	const float DevignetteStrength = Params1.z;
+	
+	float4 Foreground = InputTexture.SampleLevel(InputSampler, UV, 0);
+#if USE_CLEAN_PLATE
+	float3 Background = CleanPlateTexture.SampleLevel(CleanPlateSampler, UV, 0).rgb;
+#else
+	float3 Background = KeyColor;
+#endif
+
+	float4 Result = Foreground;
+
+	float2 ScreenPos = UVToScreenPos(UV, Input_Extent);
+	float AspectRatio = Output_ViewportSize.y * Output_ViewportSizeInverse.x;
+	float VignetteMask = Vignette(ScreenPos, AspectRatio, DevignetteStrength);
+
+	if(DevignetteStrength > ZeroEpsilon)
+	{
+		Result.rgb *= rcp(VignetteMask);
+	}
+
+	BRANCH
+	if(ScreenType == 0)
+	{
+		Result = RedScreenKeyer(Result.rgb, Background, ChannelWeights[1], ChannelWeights[2]);
+	}
+	else if(ScreenType == 1)
+	{
+		Result = GreenScreenKeyer(Result.rgb, Background, ChannelWeights[0], ChannelWeights[2]);
+	}
+	else if(ScreenType == 2)
+	{
+		Result = BlueScreenKeyer(Result.rgb, Background, ChannelWeights[0], ChannelWeights[1]);
+	}
+
+	BRANCH
+	if(DespillStrength > 0.0f)
+	{
+		// Map the source color onto what we call the "foreground".
+		// Convert both to Luv color space
+		float3 Foreground_Luv = RGB_to_Luv(Result.rgb);
+		float3 Background_Luv = RGB_to_Luv(Background);
+
+		// Project the foreground color onto the background in uv-space
+		float2 ProjectedDiff = ProjectClamped(Foreground_Luv.yz, Background_Luv.yz);
+
+		// Despill the foreground
+		Foreground_Luv.yz -= DespillStrength * ProjectedDiff;
+
+    	// Convert back to RGB
+		Result.rgb = Luv_to_RGB(Foreground_Luv);
+	}
+
+	if(bPreserveVignetteAfterKey > 0 && DevignetteStrength > ZeroEpsilon)
+	{
+		Result.rgb *= VignetteMask;
+	}
+
+	// Apply final alpha threshold
+	Result.a = Linearstep(AlphaMin, AlphaMax, Result.a);
+	
+	const uint2 PixelPos = DispatchThreadId + Output_ViewportMin;
+	RWOutputTexture[PixelPos] = Output(Result);
+}
diff --git a/Engine/Plugins/Experimental/Compositing/Composite/Shaders/Private/CompositeColorKeyerCommon.ush b/Engine/Plugins/Experimental/Compositing/Composite/Shaders/Private/CompositeColorKeyerCommon.ush
new file mode 100644
--- /dev/null
+++ b/Engine/Plugins/Experimental/Compositing/Composite/Shaders/Private/CompositeColorKeyerCommon.ush
@@ -0,0 +1,311 @@
+// Copyright Epic Games, Inc. All Rights Reserved.
+
+#pragma once
+
+#include "/Engine/Private/ACES/ACESCommon.ush"
+#include "/Engine/Private/ColorSpace.ush"
+#include "/Engine/Private/GammaCorrectionCommon.ush"
+#include "/Engine/Private/PostProcessCommon.ush"
+
+// CIE 1931 D65 white point
+static const float3 XYZ_n = xyY_2_XYZ(float3(0.3127, 0.3290, 1.0));
+static const float ZeroEpsilon = 1.0e-6;
+static const float DeltaE_JND = 1.0;
+static const float Lstar_CIELUV = 100.0;
+static const float Lstar_CIELUV_Inv = 1.0 / Lstar_CIELUV;
+
+// Safe reciprocal
+float SafeRcp(float x)
+{
+	return x > 0.0 ? rcp(x) : 0.0;
+}
+
+float LuminanceSRGB(float3 InColor)
+{
+	return dot(InColor, float3(0.2126390059, 0.7151686788, 0.0721923154));
+}
+
+// Linear equivalent of smoothstep
+float Linearstep(float Edge0, float Edge1, float X)
+{
+	return saturate((X - Edge0) / (Edge1 - Edge0));
+}
+
+float2 Project(float2 VecA, float2 VecB)
+{
+	float AdotB = dot(VecA, VecB);
+	float BdotB = dot(VecB, VecB);
+	
+	return select(BdotB < ZeroEpsilon, float2(0, 0), VecB * (AdotB / BdotB));
+}
+
+// Project vector A onto B, and return 0 if the projection was negative
+float2 ProjectClamped(float2 VecA, float2 VecB)
+{
+	float AdotB = max(0.0f, dot(VecA, VecB));
+	float BdotB = dot(VecB, VecB);
+	
+	return select(BdotB < ZeroEpsilon, float2(0, 0), VecB * (AdotB / BdotB));
+}
+
+float3 ProjectClamped(float3 VecA, float3 VecB)
+{
+	float AdotB = max(0.0f, dot(VecA, VecB));
+	float BdotB = dot(VecB, VecB);
+
+	return select(BdotB < ZeroEpsilon, float3(0, 0, 0), VecB * (AdotB / BdotB));
+}
+
+float GetAngle(float2 A, float2 B)
+{
+	return acos(dot(normalize(A), normalize(B)));
+}
+
+float2 Rotate2D(float2 V, float Theta)
+{
+	float S, C;
+	sincos(Theta, S, C);
+	return float2(C*V.x - S*V.y,
+				  S*V.x + C*V.y);
+}
+
+// Convert color from XYZ to CIE 1976 L*, u*, v* color space.
+// We use the latter for its improved perceptual uniformity, as a way to measure chromaticity distance for keying.
+float3 CIE_XYZ_to_Luv(float3 XYZ)
+{
+	const float Epsilon = 216.0f / 24389.0f; // (6/29)^3
+	const float Kappa = 24389.0f / 27.0f; // (29/3)^3
+
+	float L = select(XYZ[1] / XYZ_n[1] > Epsilon, 116 * pow(XYZ[1]/ XYZ_n[1], 1.0/3.0) - 16, Kappa * XYZ[1] / XYZ_n[1]);
+
+	if (L < ZeroEpsilon)
+	{
+		return float3(0.0f, 0.0f, 0.0f);
+	}
+	
+	float Denom = XYZ[0] + 15 * XYZ[1] + 3 * XYZ[2];
+	float Denom_n = XYZ_n[0] + 15 * XYZ_n[1] + 3 * XYZ_n[2];
+	
+	float up = select(Denom > 0.0f, 4 * XYZ[0] / Denom, 0.0f);
+	float vp = select(Denom > 0.0f, 9 * XYZ[1] / Denom, 0.0f);
+
+	float u_n = 4 * XYZ_n[0] / Denom_n;
+	float v_n = 9 * XYZ_n[1] / Denom_n;
+
+	float u = 13 * L * (up - u_n);
+	float v = 13 * L * (vp - v_n);
+
+	return float3(L, u, v);
+}
+
+// Convert color from CIE 1976 L*, u*, v* to XYZ color space.
+float3 CIE_Luv_to_XYZ(float3 Luv)
+{
+	const float Epsilon = 216.0f / 24389.0f; // (6/29)^3
+	const float Kappa = 24389.0f / 27.0f; // (29/3)^3
+
+	float L = Luv[0];
+	float u = Luv[1];
+	float v = Luv[2];
+
+	// Handle the case where L is zero, which corresponds to black.
+	if (L < ZeroEpsilon)
+	{
+		return float3(0.0f, 0.0f, 0.0f);
+	}	
+
+	float Y_ratio = select(L > Kappa * Epsilon, pow((L + 16.0f) / 116.0f, 3.0f), L / Kappa);
+	float Y = Y_ratio * XYZ_n[1];
+
+	float Denom_n = XYZ_n[0] + 15.0f * XYZ_n[1] + 3.0f * XYZ_n[2];
+	float u_n = 4.0f * XYZ_n[0] / Denom_n;
+	float v_n = 9.0f * XYZ_n[1] / Denom_n;
+	
+	float up = u / (13.0f * L) + u_n;
+	float vp = v / (13.0f * L) + v_n;
+
+	float X = Y * (9.0f * up) / (4.0f * vp);
+	float Z = Y * (12.0f - 3.0f * up - 20.0f * vp) / (4.0f * vp);
+	
+	return float3(X, Y, Z);
+}
+
+float3 RGB_to_Luv(float3 InRGB)
+{
+	return CIE_XYZ_to_Luv(LinearRGB_2_XYZ(InRGB));
+}
+
+float3 Luv_to_RGB(float3 InLuv)
+{
+	return XYZ_2_LinearRGB(CIE_Luv_to_XYZ(InLuv));
+}
+
+float3 GreenDespillAvg(float3 Color)
+{
+	float Avg = 0.5f * (Color.r + Color.b);
+
+	return float3(
+		Color.r,
+		select(Color.g > Avg, Avg, Color.g),
+		Color.b
+	);
+}
+
+float3 BlueDespillAvg(float3 Color)
+{
+	float Avg = 0.5f * (Color.r + Color.g);
+
+	return float3(
+		Color.r,
+		Color.g,
+		select(Color.b > Avg, Avg, Color.b)
+	);
+}
+
+float Vignette(float2 ScreenPos, float AspectRatio, float Intensity)
+{
+	// Scale vignette to always be a circle with consistent corner intensity.
+	float2 Vignette = VignetteSpace(ScreenPos, AspectRatio);
+	return ComputeVignetteMask(Vignette, Intensity);
+}
+
+float4 GreenScreenKeyer(
+	float3 Foreground,
+	float3 Background,
+	float  RedWeight,
+	float  BlueWeight)
+{
+	// Shared weighted sums
+	float wF = RedWeight * Foreground.r + BlueWeight * Foreground.b;
+	float wB = RedWeight * Background.r + BlueWeight * Background.b;
+
+	// Green minus weighted red/blue
+	float Fg = Foreground.g - wF;
+	float Bg = Background.g - wB;
+
+	// Ratios
+	float InvBg       = SafeRcp(Bg);
+	float GreenRatio  = saturate(Fg * InvBg);
+	float Alpha       = 1 - GreenRatio;
+
+	// Built-in "despill"
+	float3 DespillF = float3(Foreground.r, wF, Foreground.b);
+	float3 DespillB = float3(Background.r, wB, Background.b);
+
+	// Masked background
+	float3 MaskedB = DespillB * GreenRatio;
+
+	// Handle screen brighter than sample color
+	float3 Delta = Foreground - Background;
+	float Mg     = Delta.g - (RedWeight * Delta.r + BlueWeight * Delta.b);
+	MaskedB.g   += saturate(Mg * InvBg * wB);
+
+	// Fill
+	float3 Fill = DespillF - MaskedB;
+
+	// Re-insert removed green
+	Fill.g += min(Fg, 0.0);
+	// Prevent negativtes
+	Fill = max(Fill, 0.0);
+
+	// Hack to make sure the bright areas of the plate are solid
+	Alpha = max(Alpha, Fill.r);
+	// Clamp alpha
+	Alpha = saturate(Alpha);
+
+	// Note: Output is already considered pre-multiplied
+	return float4(Fill, Alpha);
+}
+
+float3 SwizzleRG(float3 InColor)
+{
+	return float3(InColor.g, InColor.r, InColor.b);
+}
+
+float3 SwizzleGB(float3 InColor)
+{
+	return float3(InColor.r, InColor.b, InColor.g);
+}
+
+float4 RedScreenKeyer(float3 Foreground, float3 Background, float GreenWeight, float BlueWeight)
+{
+	float3 SwizzledForeground = SwizzleRG(Foreground);
+	float3 SizzledBackground = SwizzleRG(Background);
+	float4 Output = GreenScreenKeyer(SwizzledForeground, SizzledBackground, GreenWeight, BlueWeight);
+	Output.rgb = SwizzleRG(Output.rgb);
+
+	return Output;
+}
+
+float4 BlueScreenKeyer(float3 Foreground, float3 Background, float RedWeight, float GreenWeight)
+{
+	float3 SwizzledForeground = SwizzleGB(Foreground);
+	float3 SizzledBackground = SwizzleGB(Background);
+	float4 Output = GreenScreenKeyer(SwizzledForeground, SizzledBackground, RedWeight, GreenWeight);
+	Output.rgb = SwizzleGB(Output.rgb);
+
+	return Output;
+}
+
+// Experimental
+float4 LuvDistanceKeyer(float3 Foreground_Luv, float3 Background_Luv, float MinDistance, float MaxDistance)
+{
+	const float kDistanceScaleFactor_CIELUV = 2 * DeltaE_JND / Lstar_CIELUV;
+	float Delta = distance(Foreground_Luv.yz, Background_Luv.yz) * kDistanceScaleFactor_CIELUV;
+
+	float Alpha = Linearstep(MinDistance, MaxDistance, Delta);
+
+	// TODO: Modulate Foreground_Luv.x
+	return float4(Foreground_Luv, Alpha);
+}
+
+// Experimental
+float4 LuvAngleKeyer(float3 Foreground_Luv, float3 Background_Luv, float Acceptance, float Scale)
+{
+	// Similar to approach described in Video Demystified  - A Handbook for the Digital Engineer
+	float AcceptanceAngle = 2 * PI * Acceptance;
+	
+	// Angle between the background and the "u" axis.
+	float AxisAngle = atan2(Background_Luv.z, Background_Luv.y);
+	
+	// Axis-aligned foreground uv
+	float2 Foreground_uvaa = Rotate2D(Foreground_Luv.yz, AxisAngle);
+
+	float X = Foreground_uvaa.x;
+	float Z_sub = abs(Foreground_uvaa.y) / tan(0.5 * AcceptanceAngle);
+	float Key = select(X < Z_sub, 0.0f, X - Z_sub);
+	
+	float Alpha = saturate(1.0 - Lstar_CIELUV_Inv * Key * Scale);
+
+	// TODO: Modulate Foreground_Luv.x
+	return float4(Foreground_Luv, Alpha);
+}
+
+// Branchless 2-value sort
+void Sort2(inout float4 A, inout float4 B)
+{
+    float4 Temp = A;
+    A = min(A, B);
+    B = max(Temp, B);
+}
+
+// Median of 9 colors (https://bertdobbelaere.github.io/median_networks.html#N9L19D7)
+float4 Median9(float4 A0, float4 A1, float4 A2, float4 A3, float4 A4, float4 A5, float4 A6, float4 A7, float4 A8)
+{
+	// For future optimization, the operations on each line could be done in parallel on separate threads or lanes.
+	Sort2(A0, A7); Sort2(A1, A2); Sort2(A3, A5); Sort2(A4, A8);
+	Sort2(A0, A2); Sort2(A1, A5); Sort2(A3, A8); Sort2(A4, A7);
+	Sort2(A0, A3); Sort2(A1, A4); Sort2(A2, A8); Sort2(A5, A7);
+	Sort2(A3, A4); Sort2(A5, A6);
+	Sort2(A2, A5); Sort2(A4, A6);
+	Sort2(A2, A3); Sort2(A4, A5);
+	Sort2(A3, A4);
+
+	return A4;
+}
+
+// Denoising with per-channel median filter of 9 colors
+float4 DenoiseMedianFilter9(const float4 V[9])
+{
+    return Median9(V[0], V[1], V[2], V[3], V[4], V[5], V[6], V[7], V[8]);
+}
diff --git a/Engine/Plugins/Experimental/Compositing/Composite/Shaders/Private/CompositeDeltaKeyer.usf b/Engine/Plugins/Experimental/Compositing/Composite/Shaders/Private/CompositeDeltaKeyer.usf
deleted file mode 100644
diff --git a/Engine/Plugins/Experimental/Compositing/Composite/Shaders/Private/CompositeDeltaKeyerCommon.ush b/Engine/Plugins/Experimental/Compositing/Composite/Shaders/Private/CompositeDeltaKeyerCommon.ush
deleted file mode 100644
diff --git a/Engine/Plugins/Experimental/Compositing/Composite/Shaders/Private/CompositeDenoise.usf b/Engine/Plugins/Experimental/Compositing/Composite/Shaders/Private/CompositeDenoise.usf
--- a/Engine/Plugins/Experimental/Compositing/Composite/Shaders/Private/CompositeDenoise.usf
+++ b/Engine/Plugins/Experimental/Compositing/Composite/Shaders/Private/CompositeDenoise.usf
@@ -2,7 +2,7 @@
 
 #include "/Engine/Private/Common.ush"
 
-#include "CompositeDeltaKeyerCommon.ush"
+#include "CompositeColorKeyerCommon.ush"
 
 #ifndef TILE_SIZE
 	#error TILE_SIZE must be defined
diff --git a/Engine/Plugins/Experimental/Compositing/Composite/Source/Composite/Private/Passes/CompositePassColorKeyer.cpp b/Engine/Plugins/Experimental/Compositing/Composite/Source/Composite/Private/Passes/CompositePassColorKeyer.cpp
new file mode 100644
--- /dev/null
+++ b/Engine/Plugins/Experimental/Compositing/Composite/Source/Composite/Private/Passes/CompositePassColorKeyer.cpp
@@ -0,0 +1,287 @@
+// Copyright Epic Games, Inc. All Rights Reserved.
+
+#include "Passes/CompositePassColorKeyer.h"
+
+#include "Engine/Texture.h"
+#include "Passes/CompositeCorePassProxy.h"
+#include "PixelShaderUtils.h"
+#include "TextureResource.h"
+
+#if WITH_EDITOR
+#include "Application/ThrottleManager.h"
+#endif
+
+DECLARE_GPU_STAT_NAMED(FCompositeDenoise, TEXT("Composite.Denoise"));
+DECLARE_GPU_STAT_NAMED(FCompositeColorKeyer, TEXT("Composite.ColorKeyer"));
+
+class FCompositePassDenoiseShader : public FGlobalShader
+{
+	DECLARE_GLOBAL_SHADER(FCompositePassDenoiseShader);
+	SHADER_USE_PARAMETER_STRUCT(FCompositePassDenoiseShader, FGlobalShader);
+
+	BEGIN_SHADER_PARAMETER_STRUCT(FParameters, )
+		SHADER_PARAMETER_RDG_TEXTURE(Texture2D, InputTexture)
+		SHADER_PARAMETER_RDG_TEXTURE_UAV(RWTexture2D<float4>, RWOutputTexture)
+		SHADER_PARAMETER(FIntPoint, Dimensions)
+	END_SHADER_PARAMETER_STRUCT()
+
+	static inline void ModifyCompilationEnvironment(const FGlobalShaderPermutationParameters& Parameters, FShaderCompilerEnvironment& OutEnvironment)
+	{
+		OutEnvironment.SetDefine(TEXT("TILE_SIZE"), ThreadGroupSize);
+
+		FGlobalShader::ModifyCompilationEnvironment(Parameters, OutEnvironment);
+	}
+
+	static const uint32 ThreadGroupSize = 16;
+};
+
+class FCompositePassColorKeyerShader : public FGlobalShader
+{
+	DECLARE_GLOBAL_SHADER(FCompositePassColorKeyerShader);
+	SHADER_USE_PARAMETER_STRUCT(FCompositePassColorKeyerShader, FGlobalShader);
+
+	class FUseCleanPlate : SHADER_PERMUTATION_BOOL("USE_CLEAN_PLATE");
+	using FPermutationDomain = TShaderPermutationDomain<FUseCleanPlate>;
+
+	BEGIN_SHADER_PARAMETER_STRUCT(FParameters, )
+		SHADER_PARAMETER_RDG_TEXTURE(Texture2D, InputTexture)
+		SHADER_PARAMETER_SAMPLER(SamplerState, InputSampler)
+		SHADER_PARAMETER_TEXTURE(Texture2D, CleanPlateTexture)
+		SHADER_PARAMETER_SAMPLER(SamplerState, CleanPlateSampler)
+		SHADER_PARAMETER_RDG_TEXTURE_UAV(RWTexture2D<float4>, RWOutputTexture)
+		SHADER_PARAMETER_STRUCT(FScreenPassTextureViewportParameters, Input)
+		SHADER_PARAMETER_STRUCT(FScreenPassTextureViewportParameters, Output)
+		SHADER_PARAMETER(FUintVector2, InputTextureSize)
+		SHADER_PARAMETER(FVector3f, KeyColor)
+		SHADER_PARAMETER(FVector4f, Params0)
+		SHADER_PARAMETER(FVector4f, Params1)
+		SHADER_PARAMETER(uint32, ScreenType)
+		SHADER_PARAMETER(uint32, Visualization)
+		SHADER_PARAMETER(uint32, bPreserveVignetteAfterKey)
+		SHADER_PARAMETER(uint32, bInvertAlpha)
+	END_SHADER_PARAMETER_STRUCT()
+
+	static inline void ModifyCompilationEnvironment(const FGlobalShaderPermutationParameters& Parameters, FShaderCompilerEnvironment& OutEnvironment)
+	{
+		OutEnvironment.SetDefine(TEXT("TILE_SIZE"), ThreadGroupSize);
+
+		FGlobalShader::ModifyCompilationEnvironment(Parameters, OutEnvironment);
+	}
+
+	static const uint32 ThreadGroupSize = 16;
+};
+
+IMPLEMENT_GLOBAL_SHADER(FCompositePassDenoiseShader, "/Plugin/Composite/Private/CompositeDenoise.usf", "MainCS", SF_Compute);
+IMPLEMENT_GLOBAL_SHADER(FCompositePassColorKeyerShader, "/Plugin/Composite/Private/CompositeColorKeyer.usf", "MainCS", SF_Compute);
+
+namespace UE
+{
+	namespace CompositeCore
+	{
+		class FCompositePassColorKeyerProxy : public FCompositeCorePassProxy
+		{
+		public:
+			IMPLEMENT_COMPOSITE_PASS(FCompositePassColorKeyerProxy);
+
+			using FCompositeCorePassProxy::FCompositeCorePassProxy;
+
+			FScreenPassTexture AddDenoisePass(FRDGBuilder& GraphBuilder, const FScreenPassTexture& Input, ERHIFeatureLevel::Type FeatureLevel) const
+			{
+				RDG_EVENT_SCOPE_STAT(GraphBuilder, FCompositeDenoise, "Composite.Denoise");
+				RDG_GPU_STAT_SCOPE(GraphBuilder, FCompositeDenoise);
+
+				FGlobalShaderMap* GlobalShaderMap = GetGlobalShaderMap(FeatureLevel);
+				const FIntPoint TextureSize = Input.Texture->Desc.Extent;
+
+				const FRDGTextureDesc OutputDesc = FRDGTextureDesc::Create2D(TextureSize, PF_FloatRGBA, FClearValueBinding::Black, TexCreate_UAV | TexCreate_ShaderResource | TexCreate_RenderTargetable);
+				FScreenPassRenderTarget Output = FScreenPassRenderTarget(GraphBuilder.CreateTexture(OutputDesc, TEXT("CompositeTextureDenoised")), ERenderTargetLoadAction::ENoAction);
+
+				FCompositePassDenoiseShader::FParameters* PassParameters = GraphBuilder.AllocParameters<FCompositePassDenoiseShader::FParameters>();
+				PassParameters->InputTexture = Input.Texture;
+				PassParameters->RWOutputTexture = GraphBuilder.CreateUAV(Output.Texture);
+				PassParameters->Dimensions = TextureSize;
+
+				TShaderMapRef<FCompositePassDenoiseShader> ComputeShader(GlobalShaderMap);
+				FComputeShaderUtils::AddPass(
+					GraphBuilder,
+					RDG_EVENT_NAME("CompositeCore.Denoise (%dx%d)", TextureSize.X, TextureSize.Y),
+					GSupportsEfficientAsyncCompute ? ERDGPassFlags::AsyncCompute : ERDGPassFlags::Compute,
+					ComputeShader,
+					PassParameters,
+					FComputeShaderUtils::GetGroupCount(TextureSize, FCompositePassDenoiseShader::ThreadGroupSize)
+				);
+
+				return Output;
+			}
+
+			FPassTexture Add(FRDGBuilder& GraphBuilder, const FSceneView& InView, const FPassInputArray& Inputs, const FPassContext& PassContext) const override
+			{
+				RDG_EVENT_SCOPE_STAT(GraphBuilder, FCompositeColorKeyer, "Composite.ColorKeyer");
+				RDG_GPU_STAT_SCOPE(GraphBuilder, FCompositeColorKeyer);
+
+				check(ValidateInputs(Inputs));
+
+				const UE::CompositeCore::FResourceMetadata& Metadata = Inputs[0].Metadata;
+				FScreenPassTexture Input = Inputs[0].Texture;
+				FScreenPassRenderTarget Output = Inputs.OverrideOutput;
+
+				if (DenoiseMethod != ECompositeDenoiseMethod::None)
+				{
+					Input = AddDenoisePass(GraphBuilder, Input, InView.GetFeatureLevel());
+				}
+
+				// If the override output is provided, it means that this is the last pass in post processing.
+				if (!Output.IsValid())
+				{
+					Output = CreateOutputRenderTarget(GraphBuilder, InView, PassContext.OutputViewRect, Input.Texture->Desc, TEXT("ColorKeyerCompositePass"));
+				}
+
+				FGlobalShaderMap* GlobalShaderMap = GetGlobalShaderMap(InView.GetFeatureLevel());
+
+				const bool bSameSize = Output.ViewRect.Size() == Input.ViewRect.Size();
+				const FScreenPassTextureViewportParameters InputParameters = GetScreenPassTextureViewportParameters(FScreenPassTextureViewport(Input));
+				const FScreenPassTextureViewportParameters OutputParameters = GetScreenPassTextureViewportParameters(FScreenPassTextureViewport(Output));
+
+				FCompositePassColorKeyerShader::FParameters* PassParameters = GraphBuilder.AllocParameters<FCompositePassColorKeyerShader::FParameters>();
+				PassParameters->InputTexture = Input.Texture;
+				PassParameters->InputSampler = bSameSize ? TStaticSamplerState<SF_Point>::GetRHI() : TStaticSamplerState<SF_Bilinear>::GetRHI();
+				PassParameters->RWOutputTexture = GraphBuilder.CreateUAV(Output.Texture);
+				PassParameters->Input = InputParameters;
+				PassParameters->Output = OutputParameters;
+				PassParameters->KeyColor = FVector3f(KeyColor);
+				PassParameters->Params0 = Params0;
+				PassParameters->Params1 = Params1;
+				PassParameters->ScreenType = static_cast<uint32>(ScreenType);
+				PassParameters->Visualization = static_cast<uint32>(Visualization);
+				PassParameters->bPreserveVignetteAfterKey = static_cast<uint32>(bPreserveVignetteAfterKey);
+				PassParameters->bInvertAlpha = static_cast<uint32>(bInvertAlpha);
+	
+				bool bUseCleanPlate = false;
+				PassParameters->CleanPlateTexture = GBlackTexture->TextureRHI;
+				PassParameters->CleanPlateSampler = TStaticSamplerState<SF_Bilinear>::GetRHI();
+
+				if (CleanPlateWeakPtr.IsValid(false, true))
+				{
+					if (const FTextureResource* TextureResource = CleanPlateWeakPtr->GetResource())
+					{
+						PassParameters->CleanPlateTexture = TextureResource->TextureRHI;
+						bUseCleanPlate = true;
+					}
+				}
+
+				FCompositePassColorKeyerShader::FPermutationDomain PermutationVector;
+				PermutationVector.Set<FCompositePassColorKeyerShader::FUseCleanPlate>(bUseCleanPlate);
+
+				const FIntPoint OutputSize = Output.ViewRect.Size();
+				TShaderMapRef<FCompositePassColorKeyerShader> ComputeShader(GlobalShaderMap, PermutationVector);
+				FComputeShaderUtils::AddPass(
+					GraphBuilder,
+					RDG_EVENT_NAME("Composite.ColorKeyer (%dx%d)", OutputSize.X, OutputSize.Y),
+					ERDGPassFlags::Compute,
+					ComputeShader,
+					PassParameters,
+					FComputeShaderUtils::GetGroupCount(OutputSize, FCompositePassColorKeyerShader::ThreadGroupSize)
+				);
+
+				return UE::CompositeCore::FPassTexture{ MoveTemp(Output), Metadata };
+			}
+
+			ECompositeColorKeyerScreenType ScreenType;
+			FLinearColor KeyColor;
+			FVector4f Params0;
+			FVector4f Params1;
+			ECompositeColorKeyerVisualization Visualization;
+			bool bPreserveVignetteAfterKey;
+			ECompositeDenoiseMethod DenoiseMethod;
+			bool bInvertAlpha;
+			TWeakObjectPtr<UTexture> CleanPlateWeakPtr;
+		};
+	}
+}
+
+UCompositePassColorKeyer::UCompositePassColorKeyer(const FObjectInitializer& ObjectInitializer)
+	: Super(ObjectInitializer)
+	, ScreenType(ECompositeColorKeyerScreenType::Green)
+	, KeyColor(FLinearColor::Green)
+	, CleanPlate(nullptr)
+	, RedWeight(0.5f)
+	, GreenWeight(0.5f)
+	, BlueWeight(0.5f)
+	, AlphaThreshold(0.0f, 1.0f)
+	, DespillStrength(0.0f)
+	, DevignetteStrength(0.0f)
+	, bPreserveVignetteAfterKey(true)
+	, DenoiseMethod(ECompositeDenoiseMethod::None)
+	, Visualization(ECompositeColorKeyerVisualization::Key)
+	, bInvertAlpha(false)
+{ 
+}
+
+UCompositePassColorKeyer::~UCompositePassColorKeyer() = default;
+
+bool UCompositePassColorKeyer::GetProxy(const UE::CompositeCore::FPassInputDecl& InputDecl, FSceneRenderingBulkObjectAllocator& InFrameAllocator, FCompositeCorePassProxy*& OutProxy) const
+{
+	using namespace UE::CompositeCore;
+
+	FCompositePassColorKeyerProxy* Proxy = InFrameAllocator.Create<FCompositePassColorKeyerProxy>(UE::CompositeCore::FPassInputDeclArray{ InputDecl });
+	Proxy->KeyColor = KeyColor;
+	Proxy->Params0 = FVector4f(RedWeight, GreenWeight, BlueWeight, DespillStrength);
+	Proxy->Params1 = FVector4f(AlphaThreshold.X, AlphaThreshold.Y, DevignetteStrength, 0.0f);
+	Proxy->ScreenType = ScreenType;
+	Proxy->bPreserveVignetteAfterKey = bPreserveVignetteAfterKey;
+	Proxy->DenoiseMethod = DenoiseMethod;
+	Proxy->bInvertAlpha = bInvertAlpha;
+	Proxy->Visualization = Visualization;
+	Proxy->CleanPlateWeakPtr = CleanPlate;
+
+	OutProxy = Proxy;
+	return true;
+}
+
+#if WITH_EDITOR
+void UCompositePassColorKeyer::PreEditChange(FProperty* PropertyThatWillChange)
+{
+	Super::PreEditChange(PropertyThatWillChange);
+
+	if (!PropertyThatWillChange)
+	{
+		return;
+	}
+
+	const FName PropertyName = PropertyThatWillChange->GetFName();
+	const bool bIsBooleanProperty = CastField<const FBoolProperty>(PropertyThatWillChange) != nullptr;
+
+	if (PropertyName != GET_MEMBER_NAME_CHECKED(ThisClass, bIsEnabled)
+		&& PropertyName != GET_MEMBER_NAME_CHECKED(ThisClass, ScreenType)
+		&& PropertyName != GET_MEMBER_NAME_CHECKED(ThisClass, KeyColor)
+		&& PropertyName != GET_MEMBER_NAME_CHECKED(ThisClass, CleanPlate)
+		&& !bIsBooleanProperty
+		)
+	{
+		/**
+		* For properties not included in the list above, we want to keep the editor ticking so that
+		* result from sliding their (SNumericEntryBox) values can be observed in real-time.
+		* 
+		* Unlike component transform customizations which prevent throttling explicitly
+		* (see .PreventThrottling(true)) we don't appear to have a metadata property for this.
+		* 
+		* As an alternative solution we directly disable throttling from the slate manager,
+		* and re-enable it once the value is set.
+		*/
+		FSlateThrottleManager::Get().DisableThrottle(true);
+		bEditorThrottleDisabled = true;
+	}
+}
+
+void UCompositePassColorKeyer::PostEditChangeProperty(FPropertyChangedEvent& PropertyChangedEvent)
+{
+	Super::PostEditChangeProperty(PropertyChangedEvent);
+
+	if (bEditorThrottleDisabled && PropertyChangedEvent.ChangeType != EPropertyChangeType::Interactive)
+	{
+		FSlateThrottleManager::Get().DisableThrottle(false);
+		bEditorThrottleDisabled = false;
+	}
+}
+#endif
+
diff --git a/Engine/Plugins/Experimental/Compositing/Composite/Source/Composite/Private/Passes/CompositePassDeltaKeyer.cpp b/Engine/Plugins/Experimental/Compositing/Composite/Source/Composite/Private/Passes/CompositePassDeltaKeyer.cpp
deleted file mode 100644
diff --git a/Engine/Plugins/Experimental/Compositing/Composite/Source/Composite/Public/Passes/CompositePassColorKeyer.h b/Engine/Plugins/Experimental/Compositing/Composite/Source/Composite/Public/Passes/CompositePassColorKeyer.h
new file mode 100644
--- /dev/null
+++ b/Engine/Plugins/Experimental/Compositing/Composite/Source/Composite/Public/Passes/CompositePassColorKeyer.h
@@ -0,0 +1,123 @@
+// Copyright Epic Games, Inc. All Rights Reserved.
+
+#pragma once
+
+#include "CompositePassBase.h"
+
+#include "CompositePassColorKeyer.generated.h"
+
+#define UE_API COMPOSITE_API
+
+/** Visualization mode, for inspecting your key.*/
+UENUM(BlueprintType)
+enum class ECompositeColorKeyerVisualization : uint8
+{
+	Key = 0,
+	Fill = 1,
+	Alpha = 2
+};
+
+/** Background screen color type (red, green or blue).*/
+UENUM(BlueprintType)
+enum class ECompositeColorKeyerScreenType : uint8
+{
+	Red = 0,
+	Green = 1,
+	Blue = 2
+};
+
+/** Denoising method applied before keying. */
+UENUM(BlueprintType)
+enum class ECompositeDenoiseMethod : uint8
+{
+	None,
+	Median3x3 UMETA(DisplayName = "Median 3x3"),
+};
+
+
+/**
+ * Color keyer with clean plate & spill removal support.
+*/
+UCLASS(MinimalAPI, BlueprintType, Blueprintable, EditInlineNew, CollapseCategories, meta=(DisplayName="Color Keyer Pass"))
+class UCompositePassColorKeyer : public UCompositePassBase
+{
+	GENERATED_BODY()
+
+public:
+	/** Constructor */
+	UE_API UCompositePassColorKeyer(const FObjectInitializer& ObjectInitializer);
+	
+	/** Destructor */
+	UE_API ~UCompositePassColorKeyer();
+
+	UE_API virtual bool GetProxy(const UE::CompositeCore::FPassInputDecl& InputDecl, FSceneRenderingBulkObjectAllocator& InFrameAllocator, FCompositeCorePassProxy*& OutProxy) const override;
+
+#if WITH_EDITOR
+	UE_API virtual void PreEditChange(FProperty* PropertyThatWillChange) override;
+	UE_API virtual void PostEditChangeProperty(FPropertyChangedEvent& PropertyChangedEvent) override;
+#endif
+
+public:
+	/** Type of screen color (required). The keyer works best against red, green or blue backgrounds. */
+	UPROPERTY(EditAnywhere, BlueprintReadWrite, Category = "")
+	ECompositeColorKeyerScreenType ScreenType;
+
+	/** Static background key color.*/
+	UPROPERTY(EditAnywhere, BlueprintReadWrite, Category = "", meta = (OnlyUpdateOnInteractionEnd, EditCondition = "CleanPlate == nullptr"))
+	FLinearColor KeyColor;
+
+	/**
+	 * Clean plate background for calculating color differences per pixel, instead of the static key color.
+	 * Resolution must match the composite actor render resolution.
+	 */
+	UPROPERTY(EditAnywhere, BlueprintReadWrite, Category = "")
+	TObjectPtr<UTexture> CleanPlate;
+
+	/** Weight of the foreground red channel contributing to the key matte hardness. */
+	UPROPERTY(EditAnywhere, BlueprintReadWrite, Category = "", meta = (UIMin = 0.0f, UIMax = 1.0f, EditCondition = "ScreenType != ECompositeColorKeyerScreenType::Red"))
+	float RedWeight;
+
+	/** Weight of the foreground green channel contributing to the key matte hardness. */
+	UPROPERTY(EditAnywhere, BlueprintReadWrite, Category = "", meta = (UIMin = 0.0f, UIMax = 1.0f, EditCondition = "ScreenType != ECompositeColorKeyerScreenType::Green"))
+	float GreenWeight;
+
+	/** Weight of the foreground blue channel contributing to the key matte hardness. */
+	UPROPERTY(EditAnywhere, BlueprintReadWrite, Category = "", meta = (UIMin = 0.0f, UIMax = 1.0f, EditCondition = "ScreenType != ECompositeColorKeyerScreenType::Blue"))
+	float BlueWeight;
+
+	/** Thresholds any alpha value outside the specified range to zero or one respectively, with linear interpolation in-between. */
+	UPROPERTY(EditAnywhere, BlueprintReadWrite, Category = "", meta = (UIMin = 0.0f, UIMax = 1.0f))
+	FVector2f AlphaThreshold;
+
+	/** Strength of the spill reduction, 0.0: none, 1.0: full. */
+	UPROPERTY(EditAnywhere, BlueprintReadWrite, Category = "", meta = (UIMin = 0.0f, UIMax = 1.0f))
+	float DespillStrength;
+
+	/** Strength of the vignette removal. Used to improve plate uniformity & remove darker corners. */
+	UPROPERTY(EditAnywhere, BlueprintReadWrite, Category = "", meta = (UIMin = 0.0f, UIMax = 1.0f))
+	float DevignetteStrength;
+
+	/** When enabled, we undo devignetting before outputting the keyed plate. */
+	UPROPERTY(EditAnywhere, BlueprintReadWrite, Category = "")
+	bool bPreserveVignetteAfterKey;
+
+	/** Denoising method applied before the keyer. */
+	UPROPERTY(EditAnywhere, BlueprintReadWrite, Category = "")
+	ECompositeDenoiseMethod DenoiseMethod;
+
+	/** Vizualize the alpha key or fill. */
+	UPROPERTY(EditAnywhere, BlueprintReadWrite, Category = "")
+	ECompositeColorKeyerVisualization Visualization;
+
+	/** Invert the alpha key. */
+	UPROPERTY(EditAnywhere, BlueprintReadWrite, Category = "")
+	bool bInvertAlpha;
+
+#if WITH_EDITOR
+	/** Flag tracking whether we disabled throttling to later restore it. */
+	bool bEditorThrottleDisabled = false;
+#endif
+};
+
+#undef UE_API
+
diff --git a/Engine/Plugins/Experimental/Compositing/Composite/Source/Composite/Public/Passes/CompositePassDeltaKeyer.h b/Engine/Plugins/Experimental/Compositing/Composite/Source/Composite/Public/Passes/CompositePassDeltaKeyer.h
deleted file mode 100644
