Composite: Fix crash due to recent refactoring while adding actors on the shadow reflection catcher.

#rnx

diff --git a/Engine/Plugins/Experimental/Compositing/Composite/Source/Composite/Private/CompositeActor.cpp b/Engine/Plugins/Experimental/Compositing/Composite/Source/Composite/Private/CompositeActor.cpp
--- a/Engine/Plugins/Experimental/Compositing/Composite/Source/Composite/Private/CompositeActor.cpp
+++ b/Engine/Plugins/Experimental/Compositing/Composite/Source/Composite/Private/CompositeActor.cpp
@@ -662,7 +662,12 @@ UCompositeShadowReflectionCatcherComponent* ACompositeActor::FindOrCreateSceneCa
 	Component->ComponentTags.Add(UE::Composite::DefaultCompositeTag);
 	AddInstanceComponent(Component);
 
-	SceneCaptureArray.EmplaceAt(InIndex, Component);
+	const int32 MinSize = InIndex + 1;
+	if (SceneCaptureArray.Num() < MinSize)
+	{
+		SceneCaptureArray.SetNumZeroed(MinSize);
+	}
+	SceneCaptureArray[InIndex] = Component;
 
 #if WITH_EDITOR
 	FModuleManager::LoadModuleChecked<FLevelEditorModule>("LevelEditor").BroadcastComponentsEdited();
diff --git a/Engine/Plugins/Experimental/Compositing/Composite/Source/Composite/Private/Layers/CompositeLayerShadowReflection.cpp b/Engine/Plugins/Experimental/Compositing/Composite/Source/Composite/Private/Layers/CompositeLayerShadowReflection.cpp
--- a/Engine/Plugins/Experimental/Compositing/Composite/Source/Composite/Private/Layers/CompositeLayerShadowReflection.cpp
+++ b/Engine/Plugins/Experimental/Compositing/Composite/Source/Composite/Private/Layers/CompositeLayerShadowReflection.cpp
@@ -32,14 +32,33 @@ void UCompositeLayerShadowReflection::BeginDestroy()
 {
 	Super::BeginDestroy();
 
-	if (USceneCaptureComponent2D* CaptureComponent = CachedSceneCapture_CG.Get())
+	if (USceneCaptureComponent2D* CaptureComponent = CachedSceneCaptures[0].Get())
 	{
 		FCompositeRenderTargetPool::Get().ReleaseTarget(CaptureComponent->TextureTarget);
 	}
-	if (USceneCaptureComponent2D* CaptureComponent = CachedSceneCapture_NoCG.Get())
+	if (USceneCaptureComponent2D* CaptureComponent = CachedSceneCaptures[1].Get())
 	{
 		FCompositeRenderTargetPool::Get().ReleaseTarget(CaptureComponent->TextureTarget);
 	}
+
+	if (AutoConfigureActors == ECompositeHiddenInSceneCaptureConfiguration::Hidden)
+	{
+		for (TSoftObjectPtr<AActor>& Actor : Actors)
+		{
+			if (Actor.IsValid())
+			{
+				if (UPrimitiveComponent* PrimitiveComponent = Actor->GetComponentByClass<UPrimitiveComponent>())
+				{
+					constexpr bool bHideActorsInSceneCapture = false;
+
+					PrimitiveComponent->SetHiddenInSceneCapture(bHideActorsInSceneCapture);
+					PrimitiveComponent->SetAffectIndirectLightingWhileHidden(bHideActorsInSceneCapture);
+					PrimitiveComponent->SetCastHiddenShadow(bHideActorsInSceneCapture);
+				}
+			}
+		}
+	}
+
 }
 
 #if WITH_EDITOR
@@ -74,11 +93,10 @@ bool UCompositeLayerShadowReflection::GetProxy(FTraversalContext& InContext, FSc
 		return false;
 	}
 
-	CachedSceneCapture_CG = FindOrCreateSceneCapture_CG(*CompositeActor);
-	CachedSceneCapture_NoCG = FindOrCreateSceneCapture_NoCG(*CompositeActor);
+	CachedSceneCaptures = FindOrCreateSceneCapturePair(*CompositeActor);
 
 	// Cached scene captures are now expected to be valid, early out if not.
-	if (!ensure(CachedSceneCapture_CG.IsValid() && CachedSceneCapture_NoCG.IsValid()))
+	if (!ensure(CachedSceneCaptures[0].IsValid() && CachedSceneCaptures[1].IsValid()))
 	{
 		return false;
 	}
@@ -86,17 +104,12 @@ bool UCompositeLayerShadowReflection::GetProxy(FTraversalContext& InContext, FSc
 	FMergePassProxy* DivisionProxy = nullptr;
 
 	{
-		// Setup division proxy
-		TStaticArray<UCompositeShadowReflectionCatcherComponent*, 2> ShadowReflectionPair;
-		ShadowReflectionPair[0] = CachedSceneCapture_CG.Get();
-		ShadowReflectionPair[1] = CachedSceneCapture_NoCG.Get();
-
 		FPassInputDeclArray Inputs;
 		Inputs.SetNum(2);
 
 		for (int32 Index = 0; Index < 2; ++Index)
 		{
-			const TObjectPtr<UCompositeShadowReflectionCatcherComponent>& Component = ShadowReflectionPair[Index];
+			UCompositeShadowReflectionCatcherComponent* Component = CachedSceneCaptures[Index].Get();
 			FCompositeRenderTargetPool::Get().ConditionalAcquireTarget(Component, Component->TextureTarget, GetRenderResolution());
 
 			FExternalTexture& ExternalTexture = InContext.ExternalTextures.AddDefaulted_GetRef();
@@ -138,19 +151,19 @@ void UCompositeLayerShadowReflection::SetActors(TArray<TSoftObjectPtr<AActor>> I
 		return;
 	}
 
-	if (!CachedSceneCapture_NoCG.IsValid())
+	if (!CachedSceneCaptures[0].IsValid() || !CachedSceneCaptures[1].IsValid())
 	{
-		CachedSceneCapture_NoCG = FindOrCreateSceneCapture_NoCG(*CompositeActor);
+		CachedSceneCaptures = FindOrCreateSceneCapturePair(*CompositeActor);
 	}
 
-	// Early exit if the cached scene depth capture is still not valid
-	if (!ensure(CachedSceneCapture_NoCG.IsValid()))
+	// Early exit if the cached scene captures are not valid
+	if (!ensure(CachedSceneCaptures[0].IsValid() && CachedSceneCaptures[1].IsValid()))
 	{
 		return;
 	}
 
 	// First, we empty hidden actors since our previous logic used them
-	CachedSceneCapture_NoCG->HiddenActors.Empty();
+	CachedSceneCaptures[1]->HiddenActors.Empty();
 
 	// Update hidden components
 	TArray<TWeakObjectPtr<UPrimitiveComponent> > HiddenComponents;
@@ -170,7 +183,7 @@ void UCompositeLayerShadowReflection::SetActors(TArray<TSoftObjectPtr<AActor>> I
 		}
 	}
 
-	CachedSceneCapture_NoCG->HiddenComponents = HiddenComponents;
+	CachedSceneCaptures[1]->HiddenComponents = HiddenComponents;
 }
 
 void UCompositeLayerShadowReflection::ConfigurePrimitivesForHiddenShadowReflectionCatching()
@@ -203,13 +216,11 @@ void UCompositeLayerShadowReflection::ConfigurePrimitivesForHiddenShadowReflecti
 	}
 }
 
-UCompositeShadowReflectionCatcherComponent* UCompositeLayerShadowReflection::FindOrCreateSceneCapture_CG(ACompositeActor& InOuter) const
-{
-	return InOuter.FindOrCreateSceneCapture<UCompositeShadowReflectionCatcherComponent>(this, 0, FName("ShadowReflectionCatcher_CG"));
-}
-
-UCompositeShadowReflectionCatcherComponent* UCompositeLayerShadowReflection::FindOrCreateSceneCapture_NoCG(ACompositeActor& InOuter) const
+TStaticArray<TWeakObjectPtr<UCompositeShadowReflectionCatcherComponent>, 2> UCompositeLayerShadowReflection::FindOrCreateSceneCapturePair(ACompositeActor& InOuter) const
 {
-	return InOuter.FindOrCreateSceneCapture<UCompositeShadowReflectionCatcherComponent>(this, 1, FName("ShadowReflectionCatcher_NoCG"));
+	return {
+		InOuter.FindOrCreateSceneCapture<UCompositeShadowReflectionCatcherComponent>(this, 0, FName("ShadowReflectionCatcher_CG")),
+		InOuter.FindOrCreateSceneCapture<UCompositeShadowReflectionCatcherComponent>(this, 1, FName("ShadowReflectionCatcher_NoCG"))
+	};
 }
 
diff --git a/Engine/Plugins/Experimental/Compositing/Composite/Source/Composite/Public/Layers/CompositeLayerShadowReflection.h b/Engine/Plugins/Experimental/Compositing/Composite/Source/Composite/Public/Layers/CompositeLayerShadowReflection.h
--- a/Engine/Plugins/Experimental/Compositing/Composite/Source/Composite/Public/Layers/CompositeLayerShadowReflection.h
+++ b/Engine/Plugins/Experimental/Compositing/Composite/Source/Composite/Public/Layers/CompositeLayerShadowReflection.h
@@ -77,22 +77,16 @@ private:
 	/** Convenience function to configure actor primitives. */
 	void ConfigurePrimitivesForHiddenShadowReflectionCatching();
 
-	/** Find or create scene capture with CG. */
-	UCompositeShadowReflectionCatcherComponent* FindOrCreateSceneCapture_CG(ACompositeActor& InOuter) const;
-
-	/** Find or create scene capture without CG. */
-	UCompositeShadowReflectionCatcherComponent* FindOrCreateSceneCapture_NoCG(ACompositeActor& InOuter) const;
+	/** Find or create scene captures with CG and without CG. */
+	TStaticArray<TWeakObjectPtr<UCompositeShadowReflectionCatcherComponent>, 2> FindOrCreateSceneCapturePair(ACompositeActor& InOuter) const;
 
 private:
 	/** List of reflection & shadow caster actors. */
 	UPROPERTY(EditAnywhere, BlueprintGetter = GetActors, BlueprintSetter = SetActors, Category = "")
 	TArray<TSoftObjectPtr<AActor>> Actors;
 
-	/** Cached pointer to the scene capture with CG. */
-	mutable TWeakObjectPtr<UCompositeShadowReflectionCatcherComponent> CachedSceneCapture_CG;
-
-	/** Cached pointer to the scene capture without CG. */
-	mutable TWeakObjectPtr<UCompositeShadowReflectionCatcherComponent> CachedSceneCapture_NoCG;
+	/** Cached pointer to the scene captures with CG and without CG. */
+	mutable TStaticArray<TWeakObjectPtr<UCompositeShadowReflectionCatcherComponent>, 2> CachedSceneCaptures;
 
 	friend class FCompositeLayerShadowReflectionCustomization;
 };
