Composite: Leave piloting and add layer buttons enabled if there is only one composite actor, which can be implicitely selected.

diff --git a/Engine/Plugins/Compositing/Composite/Source/CompositeEditor/Private/UI/SCompositePanelLayerTree.cpp b/Engine/Plugins/Compositing/Composite/Source/CompositeEditor/Private/UI/SCompositePanelLayerTree.cpp
--- a/Engine/Plugins/Compositing/Composite/Source/CompositeEditor/Private/UI/SCompositePanelLayerTree.cpp
+++ b/Engine/Plugins/Compositing/Composite/Source/CompositeEditor/Private/UI/SCompositePanelLayerTree.cpp
@@ -1176,11 +1176,21 @@ void SCompositePanelLayerTree::RefreshAndExpandTreeView()
 
 bool SCompositePanelLayerTree::CanAddLayer() const
 {
+	if (CompositeActorTreeItems.Num() == 1)
+	{
+		return CompositeActorTreeItems[0]->CompositeActor.IsValid();
+	}
+
 	return TreeView.IsValid() && !TreeView->GetSelectedItems().IsEmpty();
 }
 
 bool SCompositePanelLayerTree::CanPilotCamera() const
 {
+	if (CompositeActorTreeItems.Num() == 1)
+	{
+		return CompositeActorTreeItems[0]->CompositeActor.IsValid();
+	}
+
 	TArray<FCompositeActorTreeItemPtr> SelectedItems = TreeView.IsValid() ? TreeView->GetSelectedItems() : TArray<FCompositeActorTreeItemPtr>();
 	if (!SelectedItems.Num())
 	{
@@ -1197,19 +1207,31 @@ bool SCompositePanelLayerTree::CanPilotCamera() const
 
 bool SCompositePanelLayerTree::IsPilotingCamera() const
 {
-	TArray<FCompositeActorTreeItemPtr> SelectedItems = TreeView.IsValid() ? TreeView->GetSelectedItems() : TArray<FCompositeActorTreeItemPtr>();
-	if (!SelectedItems.Num())
+	TWeakObjectPtr<ACompositeActor> CompositeActorWeak = nullptr;
+
+	if (CompositeActorTreeItems.Num() == 1)
 	{
-		return false;
+		CompositeActorWeak = CompositeActorTreeItems[0]->CompositeActor;
+	}
+	else
+	{
+		TArray<FCompositeActorTreeItemPtr> SelectedItems = TreeView.IsValid() ? TreeView->GetSelectedItems() : TArray<FCompositeActorTreeItemPtr>();
+		if (!SelectedItems.IsEmpty())
+		{
+			CompositeActorWeak = SelectedItems[0]->CompositeActor;
+		}
 	}
 
-	if (!SelectedItems[0]->HasValidCompositeActor())
+	TStrongObjectPtr CompositeActor = CompositeActorWeak.Pin();
+
+	if (!CompositeActor.IsValid())
 	{
 		return false;
 	}
 
-	AActor* CompositeCameraActor = SelectedItems[0]->CompositeActor->GetCamera().OtherActor.Get();
-	if (!CompositeCameraActor)
+	AActor* CompositeCameraActor = CompositeActor->GetCamera().OtherActor.Get();
+
+	if (!IsValid(CompositeCameraActor))
 	{
 		return false;
 	}
@@ -1224,19 +1246,31 @@ bool SCompositePanelLayerTree::IsPilotingCamera() const
 
 void SCompositePanelLayerTree::OnPilotCameraToggled(bool bPilotCamera)
 {
-	TArray<FCompositeActorTreeItemPtr> SelectedItems = TreeView.IsValid() ? TreeView->GetSelectedItems() : TArray<FCompositeActorTreeItemPtr>();
-	if (!SelectedItems.Num())
+	TWeakObjectPtr<ACompositeActor> CompositeActorWeak = nullptr;
+
+	if (CompositeActorTreeItems.Num() == 1)
 	{
-		return;
+		CompositeActorWeak = CompositeActorTreeItems[0]->CompositeActor;
+	}
+	else
+	{
+		TArray<FCompositeActorTreeItemPtr> SelectedItems = TreeView.IsValid() ? TreeView->GetSelectedItems() : TArray<FCompositeActorTreeItemPtr>();
+		if (!SelectedItems.IsEmpty())
+		{
+			CompositeActorWeak = SelectedItems[0]->CompositeActor;
+		}
 	}
 
-	if (!SelectedItems[0]->HasValidCompositeActor())
+	TStrongObjectPtr CompositeActor = CompositeActorWeak.Pin();
+
+	if (!CompositeActor.IsValid())
 	{
 		return;
 	}
 
-	AActor* CompositeCameraActor = SelectedItems[0]->CompositeActor->GetCamera().OtherActor.Get();
-	if (!CompositeCameraActor)
+	AActor* CompositeCameraActor = CompositeActor->GetCamera().OtherActor.Get();
+
+	if (!IsValid(CompositeCameraActor))
 	{
 		return;
 	}
@@ -1571,19 +1605,27 @@ TSharedPtr<SWidget> SCompositePanelLayerTree::CreateTreeContextMenu()
 
 void SCompositePanelLayerTree::OnLayerAdded(const UClass* InLayerClass)
 {
+	TWeakObjectPtr<ACompositeActor> CompositeActorWeak = nullptr;
+	int32 NewLayerIndex = 0;
+
 	TArray<FCompositeActorTreeItemPtr> SelectedItems = TreeView.IsValid() ? TreeView->GetSelectedItems() : TArray<FCompositeActorTreeItemPtr>();
-	if (!SelectedItems.Num() || !SelectedItems[0].IsValid())
+	if (!SelectedItems.IsEmpty() && SelectedItems[0].IsValid())
 	{
-		return;
+		CompositeActorWeak = SelectedItems[0]->CompositeActor;
+		NewLayerIndex = SelectedItems[0]->LayerIndex + 1;
+	}
+	else if (CompositeActorTreeItems.Num() == 1)
+	{
+		CompositeActorWeak = CompositeActorTreeItems[0]->CompositeActor;
 	}
 
-	TStrongObjectPtr<ACompositeActor> CompositeActor = SelectedItems[0]->CompositeActor.Pin();
+	TStrongObjectPtr<ACompositeActor> CompositeActor = CompositeActorWeak.Pin();
+	
 	if (!CompositeActor.IsValid())
 	{
 		return;
 	}
 
-	const int32 NewLayerIndex = SelectedItems[0]->LayerIndex + 1;
 	FScopedTransaction LayerAddedTransaction(LOCTEXT("LayerAddedTransaction", "Add Layer"));
 	CompositeActor->Modify();
 
