Composite: Reset view transform and field of view after exiting piloting.

diff --git a/Engine/Plugins/Compositing/Composite/Source/CompositeEditor/Private/UI/SCompositePanelLayerTree.cpp b/Engine/Plugins/Compositing/Composite/Source/CompositeEditor/Private/UI/SCompositePanelLayerTree.cpp
--- a/Engine/Plugins/Compositing/Composite/Source/CompositeEditor/Private/UI/SCompositePanelLayerTree.cpp
+++ b/Engine/Plugins/Compositing/Composite/Source/CompositeEditor/Private/UI/SCompositePanelLayerTree.cpp
@@ -7,6 +7,7 @@
 #include "Editor.h"
 #include "EngineUtils.h"
 #include "Factories.h"
+#include "LevelEditor.h"
 #include "LevelEditorSubsystem.h"
 #include "ScopedTransaction.h"
 #include "SPositiveActionButton.h"
@@ -17,6 +18,7 @@
 #include "Framework/Commands/GenericCommands.h"
 #include "HAL/PlatformApplicationMisc.h"
 #include "Layers/CompositeLayerBase.h"
+#include "SLevelViewport.h"
 #include "Styling/SlateIconFinder.h"
 #include "Styling/ToolBarStyle.h"
 #include "UObject/Class.h"
@@ -1277,14 +1279,46 @@ void SCompositePanelLayerTree::OnPilotCameraToggled(bool bPilotCamera)
 	
 	if (ULevelEditorSubsystem* LevelEditorSubsystem = GEditor->GetEditorSubsystem<ULevelEditorSubsystem>())
 	{
+		// The subsytem currently has no method for capturing the view transform, so we extract it ourselves for now.
+		FLevelEditorModule& LevelEditorModule = FModuleManager::GetModuleChecked<FLevelEditorModule>("LevelEditor");
+		TSharedPtr<SLevelViewport> ActiveLevelViewport = nullptr;
+		TSharedPtr<ILevelEditor> LevelEditor = LevelEditorModule.GetFirstLevelEditor();
+		if (LevelEditor.IsValid())
+		{
+			ActiveLevelViewport = LevelEditor->GetActiveViewportInterface();
+		}
+
 		if (bPilotCamera)
 		{
+			if (ActiveLevelViewport.IsValid())
+			{
+				FLevelEditorViewportClient& Client = ActiveLevelViewport->GetLevelViewportClient();
+				CachedPerspectiveCameraTransform = Client.GetViewTransform();
+				CachedViewFOV = Client.ViewFOV;
+			}
+
 			LevelEditorSubsystem->PilotLevelActor(CompositeCameraActor);
 			LevelEditorSubsystem->SetExactCameraView(true);
 		}
 		else
 		{
 			LevelEditorSubsystem->EjectPilotLevelActor();
+
+			if (ActiveLevelViewport.IsValid())
+			{
+				// Restore the previously cached view transform
+				FLevelEditorViewportClient& Client = ActiveLevelViewport->GetLevelViewportClient();
+				if (CachedPerspectiveCameraTransform.IsSet())
+				{
+					Client.SetViewLocation(CachedPerspectiveCameraTransform.GetValue().GetLocation());
+					Client.SetViewRotation(CachedPerspectiveCameraTransform.GetValue().GetRotation());
+				}
+
+				if(CachedViewFOV.IsSet())
+				{
+					Client.ViewFOV = CachedViewFOV.GetValue();
+				}
+			}
 		}
 	}
 }
diff --git a/Engine/Plugins/Compositing/Composite/Source/CompositeEditor/Private/UI/SCompositePanelLayerTree.h b/Engine/Plugins/Compositing/Composite/Source/CompositeEditor/Private/UI/SCompositePanelLayerTree.h
--- a/Engine/Plugins/Compositing/Composite/Source/CompositeEditor/Private/UI/SCompositePanelLayerTree.h
+++ b/Engine/Plugins/Compositing/Composite/Source/CompositeEditor/Private/UI/SCompositePanelLayerTree.h
@@ -3,7 +3,9 @@
 #pragma once
 
 #include "EditorUndoClient.h"
+#include "EditorViewportClient.h"
 #include "Layers/CompositeLayerBase.h"
+#include "Misc/Optional.h"
 #include "Widgets/SCompoundWidget.h"
 
 class FUICommandList;
@@ -180,6 +182,15 @@ private:
 
 	/** Last active actor. */
 	TWeakObjectPtr<ACompositeActor> LastActiveActor;
+
+	/**
+	 * Used to store last perspective camera transform before piloting (Actor Lock)
+	 * Once piloting ends, this transform is re-applied to the camera.
+	 */
+	TOptional<FViewportCameraTransform> CachedPerspectiveCameraTransform;
+
+	/** Cached field of view to re-apply once piloting ends. */
+	TOptional<float> CachedViewFOV;
 	
 	friend class SCompositePanelLayerTreeItemRow;
 	friend class SCompositePanelLayerTreeToolbar;
