Composite: Prevent new types from being filtered out in multi-user.

* Default instanced passes & layers to transactional to support undos.
* We also remove legacy composure descriptions in concert sync core settings.

diff --git a/Engine/Plugins/Compositing/Composite/Composite.uplugin b/Engine/Plugins/Compositing/Composite/Composite.uplugin
--- a/Engine/Plugins/Compositing/Composite/Composite.uplugin
+++ b/Engine/Plugins/Compositing/Composite/Composite.uplugin
@@ -47,6 +47,10 @@
 		{
 			"Name": "ConcertSyncClient",
 			"Enabled": true
+		},
+		{
+			"Name": "ConcertSyncCore",
+			"Enabled": true
 		}
 	]
 }
\ No newline at end of file
diff --git a/Engine/Plugins/Compositing/Composite/Source/Composite/Composite.Build.cs b/Engine/Plugins/Compositing/Composite/Source/Composite/Composite.Build.cs
--- a/Engine/Plugins/Compositing/Composite/Source/Composite/Composite.Build.cs
+++ b/Engine/Plugins/Compositing/Composite/Source/Composite/Composite.Build.cs
@@ -31,7 +31,8 @@ public class Composite : ModuleRules
 				"Renderer",
 				"RHI",
 				"CameraCalibrationCore",
-				"MediaFrameworkUtilities"
+				"MediaFrameworkUtilities",
+				"ConcertSyncCore",
 			}
 		);
 
diff --git a/Engine/Plugins/Compositing/Composite/Source/Composite/Private/CompositeModule.cpp b/Engine/Plugins/Compositing/Composite/Source/Composite/Private/CompositeModule.cpp
--- a/Engine/Plugins/Compositing/Composite/Source/Composite/Private/CompositeModule.cpp
+++ b/Engine/Plugins/Compositing/Composite/Source/Composite/Private/CompositeModule.cpp
@@ -2,6 +2,7 @@
 
 #include "CompositeModule.h"
 
+#include "ConcertSyncSettings.h"
 #include "Interfaces/IPluginManager.h"
 #include "Misc/Paths.h"
 #include "Modules/ModuleManager.h"
@@ -16,13 +17,40 @@ void FCompositeModule::StartupModule()
 	AddShaderSourceDirectoryMapping(TEXT("/Plugin/Composite"), FPaths::Combine(IPluginManager::Get().FindPlugin(TEXT("Composite"))->GetBaseDir(), TEXT("Shaders")));
 
 	FCoreDelegates::OnEnginePreExit.AddRaw(this, &FCompositeModule::OnEnginePreExit);
+	FCoreDelegates::OnAllModuleLoadingPhasesComplete.AddRaw(this, &FCompositeModule::OnAllModuleLoadingPhasesComplete);
 }
 
 void FCompositeModule::ShutdownModule()
 {
+	FCoreDelegates::OnAllModuleLoadingPhasesComplete.RemoveAll(this);
 	FCoreDelegates::OnEnginePreExit.RemoveAll(this);
 }
 
+void FCompositeModule::OnAllModuleLoadingPhasesComplete()
+{
+	// Update MU filter
+	UConcertSyncConfig* SyncConfig = UConcertSyncConfig::Get();
+	if (SyncConfig)
+	{
+		TArray<FSoftClassPath> IncludedPaths;
+		IncludedPaths.Reserve(3);
+
+		IncludedPaths.Add(FSoftClassPath(TEXT("/Script/Composite.CompositeAssetUserData")));
+		IncludedPaths.Add(FSoftClassPath(TEXT("/Script/Composite.CompositePassBase")));
+		IncludedPaths.Add(FSoftClassPath(TEXT("/Script/Composite.CompositeLayerBase")));
+
+		const FSoftClassPath OuterClassPath = FSoftClassPath(TEXT("/Script/Engine.World"));
+
+		for (const FSoftClassPath& IncludedPath : IncludedPaths)
+		{
+			FTransactionClassFilter Filter;
+			Filter.ObjectOuterClass = OuterClassPath;
+			Filter.ObjectClasses.Add(IncludedPath);
+			SyncConfig->IncludeObjectClassFilters.Add(Filter);
+		}
+	}
+}
+
 void FCompositeModule::OnEnginePreExit()
 {
 	FCompositeRenderTargetPool::Get().Empty();
diff --git a/Engine/Plugins/Compositing/Composite/Source/Composite/Private/Layers/CompositeLayerBase.cpp b/Engine/Plugins/Compositing/Composite/Source/Composite/Private/Layers/CompositeLayerBase.cpp
--- a/Engine/Plugins/Compositing/Composite/Source/Composite/Private/Layers/CompositeLayerBase.cpp
+++ b/Engine/Plugins/Compositing/Composite/Source/Composite/Private/Layers/CompositeLayerBase.cpp
@@ -8,6 +8,11 @@
 UCompositeLayerBase::UCompositeLayerBase(const FObjectInitializer& ObjectInitializer)
 	: Super(ObjectInitializer)
 {
+	// Ensure instances are always transactional
+	if (!HasAnyFlags(RF_ClassDefaultObject))
+	{
+		SetFlags(RF_Transactional);
+	}
 }
 
 UCompositeLayerBase::~UCompositeLayerBase() = default;
diff --git a/Engine/Plugins/Compositing/Composite/Source/Composite/Private/Passes/CompositePassBase.cpp b/Engine/Plugins/Compositing/Composite/Source/Composite/Private/Passes/CompositePassBase.cpp
--- a/Engine/Plugins/Compositing/Composite/Source/Composite/Private/Passes/CompositePassBase.cpp
+++ b/Engine/Plugins/Compositing/Composite/Source/Composite/Private/Passes/CompositePassBase.cpp
@@ -7,7 +7,11 @@
 UCompositePassBase::UCompositePassBase(const FObjectInitializer& ObjectInitializer)
 	: Super(ObjectInitializer)
 {
-
+	// Ensure instances are always transactional
+	if (!HasAnyFlags(RF_ClassDefaultObject))
+	{
+		SetFlags(RF_Transactional);
+	}
 }
 
 UCompositePassBase::~UCompositePassBase() = default;
diff --git a/Engine/Plugins/Compositing/Composite/Source/Composite/Public/CompositeModule.h b/Engine/Plugins/Compositing/Composite/Source/Composite/Public/CompositeModule.h
--- a/Engine/Plugins/Compositing/Composite/Source/Composite/Public/CompositeModule.h
+++ b/Engine/Plugins/Compositing/Composite/Source/Composite/Public/CompositeModule.h
@@ -19,6 +19,8 @@ public:
 	// End of IModuleInterface interface
 
 private:
+	// Engine all module loaded callback
+	void OnAllModuleLoadingPhasesComplete();
 	// Engine pre-exit callback
 	void OnEnginePreExit();
 };
diff --git a/Engine/Plugins/Developer/Concert/ConcertSync/ConcertSyncCore/Config/BaseConcertSyncCore.ini b/Engine/Plugins/Developer/Concert/ConcertSync/ConcertSyncCore/Config/BaseConcertSyncCore.ini
--- a/Engine/Plugins/Developer/Concert/ConcertSync/ConcertSyncCore/Config/BaseConcertSyncCore.ini
+++ b/Engine/Plugins/Developer/Concert/ConcertSync/ConcertSyncCore/Config/BaseConcertSyncCore.ini
@@ -11,8 +11,6 @@
 +AllowedTransientProperties="/Script/Engine.SceneComponent:AttachChildren"
 +ExcludedProperties="/Script/Engine.InstanceCacheDataBase:UniqueTransientPackage"
 +ExcludeTransactionClassFilters=(ObjectOuterClass=None,ObjectClasses=(Class'/Script/Landscape.LandscapeProxy',Class'/Script/Landscape.LandscapeHeightfieldCollisionComponent',Class'/Script/Landscape.LandscapeComponent'))
-; If you wish to use Composure in Multi-user, you may add that line under [/Script/ConcertSyncCore.ConcertSyncConfig] in your project engine config settings.
-;+IncludeObjectClassFilters=(ObjectOuterClass=None,ObjectClasses=(Class'/Script/Composure.CompositingElementPass',Class'/Script/Composure.ComposurePostProcessPassPolicy'))
 
 +PostReplicateActions=(Action=/Script/ConcertSyncCore.ConcertReplicationAction_UpdateComponentToWorld(),Properties=("/Script/Engine.SceneComponent:RelativeLocation","/Script/Engine.SceneComponent:RelativeRotation","/Script/Engine.SceneComponent:RelativeScale3D","/Script/Engine.SceneComponent:bAbsoluteLocation","/Script/Engine.SceneComponent:bAbsoluteRotation","/Script/Engine.SceneComponent:bAbsoluteScale"))
 +PostReplicateActions=(Action=/Script/ConcertSyncCore.ConcertReplicationAction_MarkRenderStateDirty(),Properties=("/Script/Engine.SceneComponent:bVisible"))
