Composure: Added filters to Plate Passes lists to prevent certain passes from showing up as options for new passes for imcompatible pass types, as well as a filter to prevent composite mesh actors from being selected as shadow/reflection layer content.

diff --git a/Engine/Plugins/Compositing/Composite/Source/CompositeEditor/Private/Customizations/CompositeLayerShadowReflectionCustomization.cpp b/Engine/Plugins/Compositing/Composite/Source/CompositeEditor/Private/Customizations/CompositeLayerShadowReflectionCustomization.cpp
--- a/Engine/Plugins/Compositing/Composite/Source/CompositeEditor/Private/Customizations/CompositeLayerShadowReflectionCustomization.cpp
+++ b/Engine/Plugins/Compositing/Composite/Source/CompositeEditor/Private/Customizations/CompositeLayerShadowReflectionCustomization.cpp
@@ -2,11 +2,14 @@
 
 #include "CompositeLayerShadowReflectionCustomization.h"
 
+#include "ActorTreeItem.h"
+#include "CompositeMeshActor.h"
 #include "DetailCategoryBuilder.h"
 #include "DetailLayoutBuilder.h"
 #include "DetailWidgetRow.h"
 #include "IDetailGroup.h"
 #include "IPropertyUtilities.h"
+#include "SceneOutlinerFilters.h"
 #include "Layers/CompositeLayerShadowReflection.h"
 #include "UI/SCompositeActorPickerTable.h"
 #include "Widgets/Text/STextBlock.h"
@@ -54,6 +57,17 @@ void FCompositeLayerShadowReflectionCustomization::CustomizeDetails(IDetailLayou
 		.WholeRowContent()
 		[
 			SNew(SCompositeActorPickerTable, ActorListRef)
+			.SceneOutlinerFilters_Lambda([]()
+			{
+				TSharedPtr<FSceneOutlinerFilters> Filters = MakeShared<FSceneOutlinerFilters>();
+
+				Filters->AddFilterPredicate<FActorTreeItem>(FActorTreeItem::FFilterPredicate::CreateLambda([](const AActor* InActor)
+				{
+					return !InActor->IsA<ACompositeMeshActor>();
+				}));
+				
+				return Filters;
+			})
 			.OnLayoutSizeChanged(this, &FCompositeLayerShadowReflectionCustomization::OnLayoutSizeChanged)
 		];
 	}
diff --git a/Engine/Plugins/Compositing/Composite/Source/CompositeEditor/Private/UI/SCompositeActorPickerSceneOutliner.cpp b/Engine/Plugins/Compositing/Composite/Source/CompositeEditor/Private/UI/SCompositeActorPickerSceneOutliner.cpp
--- a/Engine/Plugins/Compositing/Composite/Source/CompositeEditor/Private/UI/SCompositeActorPickerSceneOutliner.cpp
+++ b/Engine/Plugins/Compositing/Composite/Source/CompositeEditor/Private/UI/SCompositeActorPickerSceneOutliner.cpp
@@ -119,6 +119,7 @@ void SCompositeActorPickerSceneOutliner::Construct(const FArguments& InArgs, con
 		InitOptions.bShowSearchBox = true;
 		InitOptions.bShowCreateNewFolder = false;
 		InitOptions.bFocusSearchBoxWhenOpened = true;
+		InitOptions.Filters = InArgs._SceneOutlinerFilters;
 
 		InitOptions.ColumnMap.Add(FCompositeActorPickerSelectedColumn::GetID(), FSceneOutlinerColumnInfo(
 			ESceneOutlinerColumnVisibility::Visible,
diff --git a/Engine/Plugins/Compositing/Composite/Source/CompositeEditor/Private/UI/SCompositeActorPickerSceneOutliner.h b/Engine/Plugins/Compositing/Composite/Source/CompositeEditor/Private/UI/SCompositeActorPickerSceneOutliner.h
--- a/Engine/Plugins/Compositing/Composite/Source/CompositeEditor/Private/UI/SCompositeActorPickerSceneOutliner.h
+++ b/Engine/Plugins/Compositing/Composite/Source/CompositeEditor/Private/UI/SCompositeActorPickerSceneOutliner.h
@@ -5,6 +5,8 @@
 #include "SCompositeActorPickerTable.h"
 #include "Widgets/SCompoundWidget.h"
 
+struct FSceneOutlinerFilters;
+
 /**
  * Scene Outliner used for picking multiple actors quickly for a composite actor list
  */
@@ -13,6 +15,7 @@ class SCompositeActorPickerSceneOutliner : public SCompoundWidget
 public:
 	SLATE_BEGIN_ARGS(SCompositeActorPickerSceneOutliner) { }
 		SLATE_EVENT(FSimpleDelegate, OnActorListChanged)
+		SLATE_ARGUMENT(TSharedPtr<FSceneOutlinerFilters>, SceneOutlinerFilters)
 	SLATE_END_ARGS()
 
 	void Construct(const FArguments& InArgs, const FCompositeActorPickerListRef& InActorListRef);
diff --git a/Engine/Plugins/Compositing/Composite/Source/CompositeEditor/Private/UI/SCompositeActorPickerTable.cpp b/Engine/Plugins/Compositing/Composite/Source/CompositeEditor/Private/UI/SCompositeActorPickerTable.cpp
--- a/Engine/Plugins/Compositing/Composite/Source/CompositeEditor/Private/UI/SCompositeActorPickerTable.cpp
+++ b/Engine/Plugins/Compositing/Composite/Source/CompositeEditor/Private/UI/SCompositeActorPickerTable.cpp
@@ -41,6 +41,7 @@ private:
 	
 public:	
 	SLATE_BEGIN_ARGS(SCompositeActorPickerToolbar) {}
+		SLATE_ATTRIBUTE(TSharedPtr<FSceneOutlinerFilters>, SceneOutlinerFilters)
 		SLATE_EVENT(FSimpleDelegate, OnFilterChanged)
 		SLATE_EVENT(FSimpleDelegate, OnActorListChanged)
 		SLATE_EVENT(FOnExtendAddMenu, OnExtendAddMenu)
@@ -49,6 +50,7 @@ public:
 	void Construct(const FArguments& InArgs, const FCompositeActorPickerListRef& InActorListRef)
 	{
 		ActorListRef = InActorListRef;
+		SceneOutlinerFilters = InArgs._SceneOutlinerFilters;
 		OnFilterChanged = InArgs._OnFilterChanged;
 		OnActorListChanged = InArgs._OnActorListChanged;
 		OnExtendAddMenu = InArgs._OnExtendAddMenu;
@@ -151,6 +153,7 @@ private:
 		{
 			TSharedRef<SWidget> SceneOutliner =
 				SNew(SCompositeActorPickerSceneOutliner, ActorListRef)
+				.SceneOutlinerFilters(SceneOutlinerFilters.Get(nullptr))
 				.OnActorListChanged(OnActorListChanged);
 			
 			MenuBuilder.AddWidget(SceneOutliner, FText::GetEmpty(), true);
@@ -228,6 +231,9 @@ private:
 	
 	TSharedPtr<SFilterSearchBox> TextFilterSearchBox;
 
+	/** Attribute to retrieve the filters to apply to the scene outliner actor picker when it is opened */
+	TAttribute<TSharedPtr<FSceneOutlinerFilters>> SceneOutlinerFilters;
+
 	FSimpleDelegate OnFilterChanged;
 	FSimpleDelegate OnActorListChanged;
 	FOnExtendAddMenu OnExtendAddMenu;
@@ -390,6 +396,7 @@ void SCompositeActorPickerTable::Construct(const FArguments& InArgs, const FComp
 		.Padding(0.0f, 4.0f)
 		[
 			SAssignNew(Toolbar, SCompositeActorPickerToolbar, ActorListRef)
+			.SceneOutlinerFilters(InArgs._SceneOutlinerFilters)
 			.OnFilterChanged(this, &SCompositeActorPickerTable::OnFilterChanged)
 			.OnActorListChanged_Lambda([this]
 			{
diff --git a/Engine/Plugins/Compositing/Composite/Source/CompositeEditor/Private/UI/SCompositeActorPickerTable.h b/Engine/Plugins/Compositing/Composite/Source/CompositeEditor/Private/UI/SCompositeActorPickerTable.h
--- a/Engine/Plugins/Compositing/Composite/Source/CompositeEditor/Private/UI/SCompositeActorPickerTable.h
+++ b/Engine/Plugins/Compositing/Composite/Source/CompositeEditor/Private/UI/SCompositeActorPickerTable.h
@@ -5,6 +5,7 @@
 #include "EditorUndoClient.h"
 #include "Widgets/SCompoundWidget.h"
 
+struct FSceneOutlinerFilters;
 class SCompositeActorPickerToolbar;
 class FMenuBuilder;
 class FUICommandList;
@@ -67,6 +68,7 @@ public:
 	
 public:
 	SLATE_BEGIN_ARGS(SCompositeActorPickerTable) { }
+		SLATE_ATTRIBUTE(TSharedPtr<FSceneOutlinerFilters>, SceneOutlinerFilters)
 		SLATE_EVENT(FOnExtendAddMenu, OnExtendAddMenu)
 		SLATE_EVENT(FSimpleDelegate, OnLayoutSizeChanged)
 	SLATE_END_ARGS()
diff --git a/Engine/Plugins/Compositing/Composite/Source/CompositeEditor/Private/UI/SCompositePassTree.cpp b/Engine/Plugins/Compositing/Composite/Source/CompositeEditor/Private/UI/SCompositePassTree.cpp
--- a/Engine/Plugins/Compositing/Composite/Source/CompositeEditor/Private/UI/SCompositePassTree.cpp
+++ b/Engine/Plugins/Compositing/Composite/Source/CompositeEditor/Private/UI/SCompositePassTree.cpp
@@ -17,6 +17,8 @@
 #include "HAL/PlatformApplicationMisc.h"
 #include "Layers/CompositeLayerPlate.h"
 #include "Misc/StringOutputDevice.h"
+#include "Passes/CompositePassDistortion.h"
+#include "Passes/CompositePassMaterial.h"
 #include "Styling/SlateIconFinder.h"
 #include "UObject/Package.h"
 #include "Widgets/Text/SInlineEditableTextBlock.h"
@@ -148,16 +150,19 @@ private:
 public:
 	DECLARE_DELEGATE(FOnFilterChanged)
 	DECLARE_DELEGATE_OneParam(FOnPassAdded, const UClass* /* InPassClass */)
+	DECLARE_DELEGATE_RetVal_OneParam(bool, FOnFilterNewPassType, const UClass*);
 	
 	SLATE_BEGIN_ARGS(SCompositePassTreeToolbar) {}
 		SLATE_EVENT(FOnFilterChanged, OnFilterChanged)
 		SLATE_EVENT(FOnPassAdded, OnPassAdded)
+		SLATE_EVENT(FOnFilterNewPassType, OnFilterNewPassType)
 	SLATE_END_ARGS()
 
 	void Construct(const FArguments& InArgs)
 	{
 		OnFilterChanged = InArgs._OnFilterChanged;
 		OnPassAdded = InArgs._OnPassAdded;
+		OnFilterNewPassType = InArgs._OnFilterNewPassType;
 
 		TArray<TSharedRef<FFilterBase<FFilterType>>> AllFilters;
 		InitializeFilters(AllFilters);
@@ -353,6 +358,14 @@ private:
 			{
 				continue;
 			}
+
+			if (OnFilterNewPassType.IsBound())
+			{
+				if (OnFilterNewPassType.Execute(BaseLayerType))
+				{
+					continue;
+				}
+			}
 			
 			MenuBuilder.AddMenuEntry(
 				BaseLayerType->GetDisplayNameText(),
@@ -374,6 +387,14 @@ private:
 				{
 					continue;
 				}
+
+				if (OnFilterNewPassType.IsBound())
+				{
+					if (OnFilterNewPassType.Execute(BaseLayerType))
+					{
+						continue;
+					}
+				}
 				
 				if (!bHasAddedSeparator)
 				{
@@ -466,6 +487,7 @@ private:
 	
 	FOnFilterChanged OnFilterChanged;
 	FOnPassAdded OnPassAdded;
+	FOnFilterNewPassType OnFilterNewPassType;
 };
 
 /** Table row that displays the pass tree item content in the tree view */
@@ -758,6 +780,15 @@ private:
 		FTreeItemPtr DraggedItem = DragDropOp->DraggedItem.Pin();
 		if (DraggedItem->Plate.IsValid())
 		{
+			if (UCompositePassBase* Pass = DraggedItem->GetPass())
+			{
+				// PPM passes can only be added to the Layer passes list
+				if (Pass->GetClass()->IsChildOf<UCompositePassMaterial>() && InTreeItem->PassType != SCompositePassTree::EPassType::Layer)
+				{
+					return TOptional<EItemDropZone>();
+				}
+			}
+			
 			if (!InTreeItem->HasValidPassIndex())
 			{
 				return EItemDropZone::BelowItem;
@@ -849,6 +880,7 @@ void SCompositePassTree::Construct(const FArguments& InArgs, const TWeakObjectPt
 			SAssignNew(Toolbar, SCompositePassTreeToolbar)
 			.OnFilterChanged(this, &SCompositePassTree::OnFilterChanged)
 			.OnPassAdded(this, &SCompositePassTree::OnPassAdded)
+			.OnFilterNewPassType(this, &SCompositePassTree::OnFilterNewPassType)
 		]
 
 		+SVerticalBox::Slot()
@@ -1278,8 +1310,15 @@ void SCompositePassTree::OnPassAdded(const UClass* InPassClass)
 	TArray<FPassTreeItemPtr> SelectedItems = TreeView.IsValid() ? TreeView->GetSelectedItems() : TArray<FPassTreeItemPtr>();
 	if (SelectedItems.IsEmpty())
 	{
-		// If there are no selected items, add the pass to the media passes list
-		SelectedItems.Add(PassTreeItems[0]);
+		// If there are no selected items, add the pass to the media passes list, unless it is a PPM pass, which can only get added to the layers passes list
+		if (InPassClass->IsChildOf<UCompositePassMaterial>())
+		{
+			SelectedItems.Add(PassTreeItems[1]);
+		}
+		else
+		{
+			SelectedItems.Add(PassTreeItems[0]);
+		}
 	}
 
 	// Item after which the new pass will get added
@@ -1289,6 +1328,12 @@ void SCompositePassTree::OnPassAdded(const UClass* InPassClass)
 		return;
 	}
 
+	// Post-process material pass can only be added to the Layer passes list
+	if (InPassClass->IsChildOf<UCompositePassMaterial>() && AnchorItem->PassType != EPassType::Layer)
+	{
+		return;
+	}
+	
 	FScopedTransaction AddPassTransaction(LOCTEXT("AddPassTransaction", "Add Pass"));
 	AnchorItem->Plate->Modify();
 	
@@ -1318,6 +1363,29 @@ void SCompositePassTree::OnPassAdded(const UClass* InPassClass)
 	}
 }
 
+bool SCompositePassTree::OnFilterNewPassType(const UClass* InPassType) const
+{
+	// Filter out Distortion passes as those are implicitly added when there is a lens file on the composite actor camera
+	if (InPassType->IsChildOf<UCompositePassDistortion>())
+	{
+		return true;
+	}
+
+	TArray<FPassTreeItemPtr> SelectedItems = TreeView.IsValid() ? TreeView->GetSelectedItems() : TArray<FPassTreeItemPtr>();
+	const bool bLayerPassesSelected = SelectedItems.IsEmpty() || SelectedItems.ContainsByPredicate([](const FPassTreeItemPtr& InSelectedItem)
+	{
+		return InSelectedItem->PassType == EPassType::Layer;
+	});
+
+	// Post-process material pass can only be added to the Layer passes list
+	if (!bLayerPassesSelected && InPassType->IsChildOf<UCompositePassMaterial>())
+	{
+		return true;
+	}
+	
+	return false;
+}
+
 FText SCompositePassTree::GetFilterStatusText() const
 {
 	int32 NumPasses = 0;
diff --git a/Engine/Plugins/Compositing/Composite/Source/CompositeEditor/Private/UI/SCompositePassTree.h b/Engine/Plugins/Compositing/Composite/Source/CompositeEditor/Private/UI/SCompositePassTree.h
--- a/Engine/Plugins/Compositing/Composite/Source/CompositeEditor/Private/UI/SCompositePassTree.h
+++ b/Engine/Plugins/Compositing/Composite/Source/CompositeEditor/Private/UI/SCompositePassTree.h
@@ -113,6 +113,9 @@ private:
 	/** Raised when a new pass is being added */
 	void OnPassAdded(const UClass* InPassClass);
 
+	/** Raised when the toolbar is adding pass types to the Add Pass menu, allowing types to be filtered out */
+	bool OnFilterNewPassType(const UClass* InPassType) const;
+	
 	/** Gets status text based on selection and active filter */
 	FText GetFilterStatusText() const;
 
