Composite: Keep disabled properties editable.

diff --git a/Engine/Plugins/Compositing/Composite/Source/Composite/Private/Layers/CompositeLayerBase.cpp b/Engine/Plugins/Compositing/Composite/Source/Composite/Private/Layers/CompositeLayerBase.cpp
--- a/Engine/Plugins/Compositing/Composite/Source/Composite/Private/Layers/CompositeLayerBase.cpp
+++ b/Engine/Plugins/Compositing/Composite/Source/Composite/Private/Layers/CompositeLayerBase.cpp
@@ -26,16 +26,10 @@ bool UCompositeLayerBase::CanEditChange(const FProperty* InProperty) const
 	if (bIsEditable && InProperty)
 	{
 		const FName PropertyName = InProperty->GetFName();
-
-		// Always return true for enabled toggle
-		if (PropertyName != GET_MEMBER_NAME_CHECKED(ThisClass, bIsEnabled))
-		{
-			bIsEditable = bIsEnabled;
-		}
 		
 		if (PropertyName == GET_MEMBER_NAME_CHECKED(ThisClass, Operation))
 		{
-			bIsEditable = bIsEnabled && !bIsSolo;
+			bIsEditable = !bIsSolo;
 		}
 	}
 
diff --git a/Engine/Plugins/Compositing/Composite/Source/Composite/Private/Passes/CompositePassBase.cpp b/Engine/Plugins/Compositing/Composite/Source/Composite/Private/Passes/CompositePassBase.cpp
--- a/Engine/Plugins/Compositing/Composite/Source/Composite/Private/Passes/CompositePassBase.cpp
+++ b/Engine/Plugins/Compositing/Composite/Source/Composite/Private/Passes/CompositePassBase.cpp
@@ -29,20 +29,7 @@ void UCompositePassBase::SetEnabled(bool bInEnabled)
 #if WITH_EDITOR
 bool UCompositePassBase::CanEditChange(const FProperty* InProperty) const
 {
-	bool bIsEditable = Super::CanEditChange(InProperty);
-
-	if (bIsEditable && InProperty)
-	{
-		const FName PropertyName = InProperty->GetFName();
-
-		// Always return true for enabled toggle
-		if (PropertyName != GET_MEMBER_NAME_CHECKED(ThisClass, bIsEnabled))
-		{
-			bIsEditable = bIsEnabled;
-		}
-	}
-
-	return bIsEditable;
+	return Super::CanEditChange(InProperty);
 }
 #endif //WITH_EDITOR
 
