Composite: Fix MU transaction for adding a new layer or pass in the Composure editor panel.

diff --git a/Engine/Plugins/Compositing/Composite/Source/CompositeEditor/Private/UI/SCompositePanelLayerTree.cpp b/Engine/Plugins/Compositing/Composite/Source/CompositeEditor/Private/UI/SCompositePanelLayerTree.cpp
--- a/Engine/Plugins/Compositing/Composite/Source/CompositeEditor/Private/UI/SCompositePanelLayerTree.cpp
+++ b/Engine/Plugins/Compositing/Composite/Source/CompositeEditor/Private/UI/SCompositePanelLayerTree.cpp
@@ -1663,7 +1663,7 @@ void SCompositePanelLayerTree::OnLayerAdded(const UClass* InLayerClass)
 	FScopedTransaction LayerAddedTransaction(LOCTEXT("LayerAddedTransaction", "Add Layer"));
 	CompositeActor->Modify();
 
-	UCompositeLayerBase* NewLayer = NewObject<UCompositeLayerBase>(CompositeActor.Get(), InLayerClass);
+	UCompositeLayerBase* NewLayer = NewObject<UCompositeLayerBase>(CompositeActor.Get(), InLayerClass, NAME_None, RF_Transactional);
 	NewLayer->Name = NewLayer->GetName();
 
 	FProperty* CompositeLayersProperty = FindFProperty<FProperty>(ACompositeActor::StaticClass(), GET_MEMBER_NAME_CHECKED(ACompositeActor, CompositeLayers));
diff --git a/Engine/Plugins/Compositing/Composite/Source/CompositeEditor/Private/UI/SCompositePassTree.cpp b/Engine/Plugins/Compositing/Composite/Source/CompositeEditor/Private/UI/SCompositePassTree.cpp
--- a/Engine/Plugins/Compositing/Composite/Source/CompositeEditor/Private/UI/SCompositePassTree.cpp
+++ b/Engine/Plugins/Compositing/Composite/Source/CompositeEditor/Private/UI/SCompositePassTree.cpp
@@ -1337,7 +1337,7 @@ void SCompositePassTree::OnPassAdded(const UClass* InPassClass)
 	FScopedTransaction AddPassTransaction(LOCTEXT("AddPassTransaction", "Add Pass"));
 	AnchorItem->Plate->Modify();
 	
-	UCompositePassBase* NewPass = NewObject<UCompositePassBase>(AnchorItem->Plate.Get(), InPassClass);
+	UCompositePassBase* NewPass = NewObject<UCompositePassBase>(AnchorItem->Plate.Get(), InPassClass, NAME_None, RF_Transactional);
 	TArray<TObjectPtr<UCompositePassBase>>& PassList = GetPassList(AnchorItem->Plate, AnchorItem->PassType);
 
 	FProperty* PassListProperty = GetPassListProperty(AnchorItem->PassType);
