Composite: Fixed undo by force-refreshing all layer state after undo/redo edits.

diff --git a/Engine/Plugins/Compositing/Composite/Source/Composite/Private/Layers/CompositeLayerSceneCapture.cpp b/Engine/Plugins/Compositing/Composite/Source/Composite/Private/Layers/CompositeLayerSceneCapture.cpp
--- a/Engine/Plugins/Compositing/Composite/Source/Composite/Private/Layers/CompositeLayerSceneCapture.cpp
+++ b/Engine/Plugins/Compositing/Composite/Source/Composite/Private/Layers/CompositeLayerSceneCapture.cpp
@@ -34,14 +34,15 @@ void UCompositeLayerSceneCapture::PostEditChangeProperty(FPropertyChangedEvent&
 	{
 		SetCustomRenderPass(bCustomRenderPass);
 	}
-	else if (PropertyName == GET_MEMBER_NAME_CHECKED(ThisClass, LayerPasses))
-	{
-		if (PropertyChangedEvent.ChangeType == EPropertyChangeType::ArrayAdd)
-		{
-			// Always add elements at the start of the array (top of the stack)
-			Algo::Rotate(LayerPasses, LayerPasses.Num() - 1);
-		}
-	}
+}
+
+void UCompositeLayerSceneCapture::PostEditUndo()
+{
+	Super::PostEditUndo();
+
+	SetActors(MoveTemp(Actors));
+
+	SetCustomRenderPass(bCustomRenderPass);
 }
 #endif
 
diff --git a/Engine/Plugins/Compositing/Composite/Source/Composite/Private/Layers/CompositeLayerShadowReflection.cpp b/Engine/Plugins/Compositing/Composite/Source/Composite/Private/Layers/CompositeLayerShadowReflection.cpp
--- a/Engine/Plugins/Compositing/Composite/Source/Composite/Private/Layers/CompositeLayerShadowReflection.cpp
+++ b/Engine/Plugins/Compositing/Composite/Source/Composite/Private/Layers/CompositeLayerShadowReflection.cpp
@@ -66,6 +66,15 @@ void UCompositeLayerShadowReflection::PostEditChangeProperty(FPropertyChangedEve
 		ConfigurePrimitivesForHiddenShadowReflectionCatching();
 	}
 }
+
+void UCompositeLayerShadowReflection::PostEditUndo()
+{
+	Super::PostEditUndo();
+
+	SetActors(MoveTemp(Actors));
+
+	ConfigurePrimitivesForHiddenShadowReflectionCatching();
+}
 #endif
 
 bool UCompositeLayerShadowReflection::GetProxy(FTraversalContext& InContext, FSceneRenderingBulkObjectAllocator& InFrameAllocator, FCompositeCorePassProxy*& OutProxy) const
diff --git a/Engine/Plugins/Compositing/Composite/Source/Composite/Private/Layers/CompositeLayerSingleLightShadow.cpp b/Engine/Plugins/Compositing/Composite/Source/Composite/Private/Layers/CompositeLayerSingleLightShadow.cpp
--- a/Engine/Plugins/Compositing/Composite/Source/Composite/Private/Layers/CompositeLayerSingleLightShadow.cpp
+++ b/Engine/Plugins/Compositing/Composite/Source/Composite/Private/Layers/CompositeLayerSingleLightShadow.cpp
@@ -173,9 +173,16 @@ void UCompositeLayerSingleLightShadow::PostEditChangeProperty(FPropertyChangedEv
 
 	if (PropertyName == GET_MEMBER_NAME_CHECKED(ThisClass, ShadowCastingActors))
 	{
-		SetShadowCastingActors(ShadowCastingActors);
+		SetShadowCastingActors(MoveTemp(ShadowCastingActors));
 	}
 }
+
+void UCompositeLayerSingleLightShadow::PostEditUndo()
+{
+	Super::PostEditUndo();
+
+	SetShadowCastingActors(MoveTemp(ShadowCastingActors));
+}
 #endif
 
 bool UCompositeLayerSingleLightShadow::GetProxy(FTraversalContext& InContext, FSceneRenderingBulkObjectAllocator& InFrameAllocator, FCompositeCorePassProxy*& OutProxy) const
diff --git a/Engine/Plugins/Compositing/Composite/Source/Composite/Public/Layers/CompositeLayerSceneCapture.h b/Engine/Plugins/Compositing/Composite/Source/Composite/Public/Layers/CompositeLayerSceneCapture.h
--- a/Engine/Plugins/Compositing/Composite/Source/Composite/Public/Layers/CompositeLayerSceneCapture.h
+++ b/Engine/Plugins/Compositing/Composite/Source/Composite/Public/Layers/CompositeLayerSceneCapture.h
@@ -27,6 +27,7 @@ public:
 
 #if WITH_EDITOR
 	UE_API virtual void PostEditChangeProperty(FPropertyChangedEvent& PropertyChangedEvent) override;
+	UE_API virtual void PostEditUndo() override;
 #endif
 
 	/** Getter function to override, returning pass layer proxies to be passed to the render thread. (Proxy objects should be allocated from the provided allocator.) */
diff --git a/Engine/Plugins/Compositing/Composite/Source/Composite/Public/Layers/CompositeLayerShadowReflection.h b/Engine/Plugins/Compositing/Composite/Source/Composite/Public/Layers/CompositeLayerShadowReflection.h
--- a/Engine/Plugins/Compositing/Composite/Source/Composite/Public/Layers/CompositeLayerShadowReflection.h
+++ b/Engine/Plugins/Compositing/Composite/Source/Composite/Public/Layers/CompositeLayerShadowReflection.h
@@ -46,6 +46,7 @@ public:
 
 #if WITH_EDITOR
 	UE_API virtual void PostEditChangeProperty(FPropertyChangedEvent& PropertyChangedEvent) override;
+	UE_API virtual void PostEditUndo() override;
 #endif
 	//~ End UObject Interface
 
diff --git a/Engine/Plugins/Compositing/Composite/Source/Composite/Public/Layers/CompositeLayerSingleLightShadow.h b/Engine/Plugins/Compositing/Composite/Source/Composite/Public/Layers/CompositeLayerSingleLightShadow.h
--- a/Engine/Plugins/Compositing/Composite/Source/Composite/Public/Layers/CompositeLayerSingleLightShadow.h
+++ b/Engine/Plugins/Compositing/Composite/Source/Composite/Public/Layers/CompositeLayerSingleLightShadow.h
@@ -39,6 +39,7 @@ public:
 
 #if WITH_EDITOR
 	UE_API virtual void PostEditChangeProperty(FPropertyChangedEvent& PropertyChangedEvent) override;
+	UE_API virtual void PostEditUndo() override;
 #endif
 	//~ End UObject Interface
 
