Composite: Fixed layer removal logic with new OnRemoved virtual function, since BeginDestroy cannot be relied on in the editor.

* Layer OnRemoved called during destruction, layer edit or undo.
* Moved scene capture component destruction to OnRemoved for consistency with FindOrCreateSceneCapture calls.
* Added warning to users that holdout composite component isn't really compatible with the new system.

diff --git a/Engine/Plugins/Compositing/Composite/Source/Composite/Private/Components/CompositeSceneCapture2DComponent.cpp b/Engine/Plugins/Compositing/Composite/Source/Composite/Private/Components/CompositeSceneCapture2DComponent.cpp
--- a/Engine/Plugins/Compositing/Composite/Source/Composite/Private/Components/CompositeSceneCapture2DComponent.cpp
+++ b/Engine/Plugins/Compositing/Composite/Source/Composite/Private/Components/CompositeSceneCapture2DComponent.cpp
@@ -69,7 +69,8 @@ void UCompositeSceneCapture2DComponent::PostEditChangeProperty(FPropertyChangedE
 const AActor* UCompositeSceneCapture2DComponent::GetViewOwner() const
 {
 	// View owner that matches the main view camera
-	if (ACompositeActor* CompositeActor = GetTypedOuter<ACompositeActor>())
+	ACompositeActor* CompositeActor = GetTypedOuter<ACompositeActor>();
+	if (IsValid(CompositeActor))
 	{
 		if (UCameraComponent* CameraComponent = Cast<UCameraComponent>(CompositeActor->GetCamera().GetComponent(nullptr)))
 		{
diff --git a/Engine/Plugins/Compositing/Composite/Source/Composite/Private/Components/CompositeViewProjectionComponent.cpp b/Engine/Plugins/Compositing/Composite/Source/Composite/Private/Components/CompositeViewProjectionComponent.cpp
--- a/Engine/Plugins/Compositing/Composite/Source/Composite/Private/Components/CompositeViewProjectionComponent.cpp
+++ b/Engine/Plugins/Compositing/Composite/Source/Composite/Private/Components/CompositeViewProjectionComponent.cpp
@@ -49,7 +49,8 @@ void UCompositeViewProjectionComponent::SetTargetComponent(const FComponentRefer
 
 void UCompositeViewProjectionComponent::InitDefaultCamera()
 {
-	if (ACompositeActor* CompositeActor = GetTypedOuter<ACompositeActor>())
+	ACompositeActor* CompositeActor = GetTypedOuter<ACompositeActor>();
+	if (IsValid(CompositeActor))
 	{
 		// Special case when the component is owned by the composite actor: we use its camera component reference directly.
 		SetTargetComponent(CompositeActor->GetCamera());
diff --git a/Engine/Plugins/Compositing/Composite/Source/Composite/Private/CompositeActor.cpp b/Engine/Plugins/Compositing/Composite/Source/Composite/Private/CompositeActor.cpp
--- a/Engine/Plugins/Compositing/Composite/Source/Composite/Private/CompositeActor.cpp
+++ b/Engine/Plugins/Compositing/Composite/Source/Composite/Private/CompositeActor.cpp
@@ -140,7 +140,7 @@ namespace UE
 				const ACompositeActor& InCompositeActor,
 				const TMap<TWeakObjectPtr<const UCompositeLayerBase>, FSceneCaptureComponentArray>& InSceneCapturesPerLayer)
 			{
-				if (InCompositeActor.CompositeLayers.IsEmpty())
+				if (InCompositeActor.GetCompositeLayers().IsEmpty())
 				{
 					return {};
 				}
@@ -148,14 +148,14 @@ namespace UE
 				UE::CompositeCore::FRenderWork RenderWork;
 				TArray<TObjectPtr<UCompositeLayerBase>> Layers;
 
-				TObjectPtr<UCompositeLayerBase> SoloLayer = FindSoloLayer(InCompositeActor.CompositeLayers);
+				TObjectPtr<UCompositeLayerBase> SoloLayer = FindSoloLayer(InCompositeActor.GetCompositeLayers());
 				if (IsValid(SoloLayer))
 				{
 					Layers.Add(SoloLayer);
 				}
 				else
 				{
-					Layers = InCompositeActor.CompositeLayers;
+					Layers = InCompositeActor.GetCompositeLayers();
 					// TODO: Remove once we have UI customization displaying layers in the correct bottom-up order
 					Algo::Reverse(Layers);
 				}
@@ -382,6 +382,20 @@ void ACompositeActor::Tick(float DeltaTime)
 	}
 }
 
+void ACompositeActor::Destroyed()
+{
+	Super::Destroyed();
+
+	for (TObjectPtr<UCompositeLayerBase>& Layer : CompositeLayers)
+	{
+		// Note: We don't use IsValid(..) since destruction should proceed even if the actor is pending kill.
+		if (Layer != nullptr)
+		{
+			Layer->OnRemoved();
+		}
+	}
+}
+
 #if WITH_EDITOR
 void ACompositeActor::PreDuplicateFromRoot(FObjectDuplicationParameters& DupParams)
 {
@@ -422,8 +436,6 @@ void ACompositeActor::PreEditChange(FProperty* PropertyAboutToChange)
 
 	Super::PreEditChange(PropertyAboutToChange);
 
-	PreEditCompositeLayers.Reset();
-
 	if (PropertyAboutToChange && PropertyAboutToChange->GetFName() == GET_MEMBER_NAME_CHECKED(ThisClass, CompositeLayers))
 	{
 		PreEditCompositeLayers = CompositeLayers;
@@ -480,15 +492,18 @@ void ACompositeActor::PostEditChangeChainProperty(struct FPropertyChangedChainEv
 
 			for (UCompositeLayerBase* PreEditCompositeLayer : PreEditCompositeLayers)
 			{
-				if (IsValid(PreEditCompositeLayer))
+				// Note: We don't use IsValid(..) since destruction should proceed even if the actor is pending kill.
+				if (PreEditCompositeLayer != nullptr)
 				{
 					if (!CompositeLayers.Contains(PreEditCompositeLayer))
 					{
-						DestroySceneCaptures(PreEditCompositeLayer);
+						PreEditCompositeLayer->OnRemoved();
 					}
 				}
 			}
 
+			PreEditCompositeLayers.Reset();
+
 			if (PropertyChangedChainEvent.PropertyChain.Num() <= 1)
 			{
 				for (int32 Index = 0; Index < CompositeLayers.Num(); ++Index)
@@ -512,6 +527,32 @@ void ACompositeActor::PostEditChangeChainProperty(struct FPropertyChangedChainEv
 	}
 }
 
+void ACompositeActor::PreEditUndo()
+{
+	Super::PreEditUndo();
+
+	PreEditCompositeLayers = CompositeLayers;
+}
+
+void ACompositeActor::PostEditUndo()
+{
+	for (UCompositeLayerBase* PreEditCompositeLayer : PreEditCompositeLayers)
+	{
+		// Note: We don't use IsValid(..) since destruction should proceed even if the actor is pending kill.
+		if (PreEditCompositeLayer != nullptr)
+		{
+			if (!CompositeLayers.Contains(PreEditCompositeLayer))
+			{
+				PreEditCompositeLayer->OnRemoved();
+			}
+		}
+	}
+
+	PreEditCompositeLayers.Reset();
+
+	Super::PostEditUndo();
+}
+
 void ACompositeActor::PostTransacted(const FTransactionObjectEvent& TransactionEvent)
 {
 	Super::PostTransacted(TransactionEvent);
@@ -621,6 +662,32 @@ void ACompositeActor::SetCamera(const FComponentReference& InComponentReference)
 	}
 }
 
+TArray<UCompositeLayerBase*> ACompositeActor::GetCompositeLayers()
+{
+	return CompositeLayers;
+}
+
+const TArray<TObjectPtr<UCompositeLayerBase>>& ACompositeActor::GetCompositeLayers() const
+{
+	return CompositeLayers;
+}
+
+void ACompositeActor::SetCompositeLayers(TArray<UCompositeLayerBase*> InLayers)
+{
+	for (UCompositeLayerBase* PreviousLayer : CompositeLayers)
+	{
+		if (PreviousLayer != nullptr)
+		{
+			if (!InLayers.Contains(PreviousLayer))
+			{
+				PreviousLayer->OnRemoved();
+			}
+		}
+	}
+
+	CompositeLayers = MoveTemp(InLayers);
+}
+
 template<>
 UCompositeSceneCapture2DComponent* ACompositeActor::FindOrCreateSceneCapture(const UCompositeLayerBase* InLayer, int32 InIndex, FName InBaseName)
 {
diff --git a/Engine/Plugins/Compositing/Composite/Source/Composite/Private/Layers/CompositeLayerBase.cpp b/Engine/Plugins/Compositing/Composite/Source/Composite/Private/Layers/CompositeLayerBase.cpp
--- a/Engine/Plugins/Compositing/Composite/Source/Composite/Private/Layers/CompositeLayerBase.cpp
+++ b/Engine/Plugins/Compositing/Composite/Source/Composite/Private/Layers/CompositeLayerBase.cpp
@@ -42,9 +42,9 @@ void UCompositeLayerBase::PostEditChangeProperty(FPropertyChangedEvent& Property
 
 	if (PropertyChangedEvent.GetMemberPropertyName() == GET_MEMBER_NAME_CHECKED(ThisClass, bIsSolo))
 	{
-		if (ACompositeActor* CompositeActor = GetTypedOuter<ACompositeActor>())
+		if (const ACompositeActor* CompositeActor = GetTypedOuter<ACompositeActor>())
 		{
-			for (TObjectPtr<UCompositeLayerBase>& CompositeLayer : CompositeActor->CompositeLayers)
+			for (const TObjectPtr<UCompositeLayerBase>& CompositeLayer : CompositeActor->GetCompositeLayers())
 			{
 				if (CompositeLayer != this)
 				{
diff --git a/Engine/Plugins/Compositing/Composite/Source/Composite/Private/Layers/CompositeLayerPlate.cpp b/Engine/Plugins/Compositing/Composite/Source/Composite/Private/Layers/CompositeLayerPlate.cpp
--- a/Engine/Plugins/Compositing/Composite/Source/Composite/Private/Layers/CompositeLayerPlate.cpp
+++ b/Engine/Plugins/Compositing/Composite/Source/Composite/Private/Layers/CompositeLayerPlate.cpp
@@ -210,10 +210,8 @@ void UCompositeLayerPlate::PostLoad()
 	}
 }
 
-void UCompositeLayerPlate::BeginDestroy()
+void UCompositeLayerPlate::OnRemoved()
 {
-	Super::BeginDestroy();
-
 	FCompositeRenderTargetPool::Get().ReleaseTarget(MediaRenderTarget);
 	FCompositeRenderTargetPool::Get().ReleaseTarget(SceneRenderTarget);
 
@@ -223,13 +221,19 @@ void UCompositeLayerPlate::BeginDestroy()
 		TryCloseMediaProfileSource();
 		
 #if WITH_EDITOR
-		// Hook into pre/post save.
 		FEditorDelegates::PreSaveWorldWithContext.RemoveAll(this);
 		FEditorDelegates::PostSaveWorldWithContext.RemoveAll(this);
 #endif // WITH_EDITOR
 	}
 }
 
+void UCompositeLayerPlate::BeginDestroy()
+{
+	Super::BeginDestroy();
+
+	OnRemoved(); // Redundant remove call for safety
+}
+
 const TArray<TSoftObjectPtr<AActor>> UCompositeLayerPlate::GetCompositeMeshes() const
 {
 	return CompositeMeshes;
@@ -767,13 +771,13 @@ void UCompositeLayerPlate::HandleEnabledStateChange(bool bInEnabled, UCompositeC
 
 bool UCompositeLayerPlate::IsCompositeMeshActorAlreadyInUse(TSoftObjectPtr<AActor> InCompositeMeshActor) const
 {
-	ACompositeActor* CompositeActor = GetTypedOuter<ACompositeActor>();
+	const ACompositeActor* CompositeActor = GetTypedOuter<ACompositeActor>();
 	if (!IsValid(CompositeActor))
 	{
 		return false;
 	}
 
-	for (const TObjectPtr<UCompositeLayerBase>& Layer : CompositeActor->CompositeLayers)
+	for (const TObjectPtr<UCompositeLayerBase>& Layer : CompositeActor->GetCompositeLayers())
 	{
 		const UCompositeLayerPlate* PlateLayer = Cast<UCompositeLayerPlate>(Layer.Get());
 		if (IsValid(PlateLayer) && PlateLayer != this)
diff --git a/Engine/Plugins/Compositing/Composite/Source/Composite/Private/Layers/CompositeLayerSceneCapture.cpp b/Engine/Plugins/Compositing/Composite/Source/Composite/Private/Layers/CompositeLayerSceneCapture.cpp
--- a/Engine/Plugins/Compositing/Composite/Source/Composite/Private/Layers/CompositeLayerSceneCapture.cpp
+++ b/Engine/Plugins/Compositing/Composite/Source/Composite/Private/Layers/CompositeLayerSceneCapture.cpp
@@ -18,6 +18,22 @@ UCompositeLayerSceneCapture::UCompositeLayerSceneCapture(const FObjectInitialize
 
 UCompositeLayerSceneCapture::~UCompositeLayerSceneCapture() = default;
 
+void UCompositeLayerSceneCapture::OnRemoved()
+{
+	ACompositeActor* CompositeActor = GetTypedOuter<ACompositeActor>();
+	// Note: We don't use IsValid(..) since destruction should proceed even if the actor is pending kill.
+	if (CompositeActor != nullptr)
+	{
+		CompositeActor->DestroySceneCaptures(this);
+	}
+}
+
+void UCompositeLayerSceneCapture::BeginDestroy()
+{
+	Super::BeginDestroy();
+
+	OnRemoved(); // Redundant remove call for safety
+}
 
 #if WITH_EDITOR
 void UCompositeLayerSceneCapture::PostEditChangeProperty(FPropertyChangedEvent& PropertyChangedEvent)
diff --git a/Engine/Plugins/Compositing/Composite/Source/Composite/Private/Layers/CompositeLayerShadowReflection.cpp b/Engine/Plugins/Compositing/Composite/Source/Composite/Private/Layers/CompositeLayerShadowReflection.cpp
--- a/Engine/Plugins/Compositing/Composite/Source/Composite/Private/Layers/CompositeLayerShadowReflection.cpp
+++ b/Engine/Plugins/Compositing/Composite/Source/Composite/Private/Layers/CompositeLayerShadowReflection.cpp
@@ -17,19 +17,8 @@ UCompositeLayerShadowReflection::UCompositeLayerShadowReflection(const FObjectIn
 
 UCompositeLayerShadowReflection::~UCompositeLayerShadowReflection() = default;
 
-void UCompositeLayerShadowReflection::BeginDestroy()
+void UCompositeLayerShadowReflection::OnRemoved()
 {
-	Super::BeginDestroy();
-
-	if (USceneCaptureComponent2D* CaptureComponent = CachedSceneCaptures[0].Get())
-	{
-		FCompositeRenderTargetPool::Get().ReleaseTarget(CaptureComponent->TextureTarget);
-	}
-	if (USceneCaptureComponent2D* CaptureComponent = CachedSceneCaptures[1].Get())
-	{
-		FCompositeRenderTargetPool::Get().ReleaseTarget(CaptureComponent->TextureTarget);
-	}
-
 	if (AutoConfigureActors == ECompositeHiddenInSceneCaptureConfiguration::Hidden)
 	{
 		for (TSoftObjectPtr<AActor>& Actor : Actors)
@@ -48,6 +37,19 @@ void UCompositeLayerShadowReflection::BeginDestroy()
 		}
 	}
 
+	ACompositeActor* CompositeActor = GetTypedOuter<ACompositeActor>();
+	// Note: We don't use IsValid(..) since destruction should proceed even if the actor is pending kill.
+	if (CompositeActor != nullptr)
+	{
+		CompositeActor->DestroySceneCaptures(this);
+	}
+}
+
+void UCompositeLayerShadowReflection::BeginDestroy()
+{
+	Super::BeginDestroy();
+
+	OnRemoved(); // Redundant remove call for safety
 }
 
 #if WITH_EDITOR
diff --git a/Engine/Plugins/Compositing/Composite/Source/Composite/Private/Layers/CompositeLayerSingleLightShadow.cpp b/Engine/Plugins/Compositing/Composite/Source/Composite/Private/Layers/CompositeLayerSingleLightShadow.cpp
--- a/Engine/Plugins/Compositing/Composite/Source/Composite/Private/Layers/CompositeLayerSingleLightShadow.cpp
+++ b/Engine/Plugins/Compositing/Composite/Source/Composite/Private/Layers/CompositeLayerSingleLightShadow.cpp
@@ -164,6 +164,23 @@ UCompositeLayerSingleLightShadow::UCompositeLayerSingleLightShadow(const FObject
 	Operation = ECompositeCoreMergeOp::Multiply;
 }
 
+void UCompositeLayerSingleLightShadow::OnRemoved()
+{
+	ACompositeActor* CompositeActor = GetTypedOuter<ACompositeActor>();
+	// Note: We don't use IsValid(..) since destruction should proceed even if the actor is pending kill.
+	if (CompositeActor != nullptr)
+	{
+		CompositeActor->DestroySceneCaptures(this);
+	}
+}
+
+void UCompositeLayerSingleLightShadow::BeginDestroy()
+{
+	Super::BeginDestroy();
+
+	OnRemoved(); // Redundant remove call for safety
+}
+
 #if WITH_EDITOR
 void UCompositeLayerSingleLightShadow::PostEditChangeProperty(FPropertyChangedEvent& PropertyChangedEvent)
 {
diff --git a/Engine/Plugins/Compositing/Composite/Source/Composite/Public/CompositeActor.h b/Engine/Plugins/Compositing/Composite/Source/Composite/Public/CompositeActor.h
--- a/Engine/Plugins/Compositing/Composite/Source/Composite/Public/CompositeActor.h
+++ b/Engine/Plugins/Compositing/Composite/Source/Composite/Public/CompositeActor.h
@@ -72,12 +72,15 @@ public:
 	UE_API virtual void UnregisterAllComponents(bool bForReregister = false) override;
 	UE_API virtual void PostUnregisterAllComponents() override;
 	UE_API virtual void Tick(float DeltaSeconds) override;
+	UE_API virtual void Destroyed() override;
 	//~ End AActor Interface
 
 #if WITH_EDITOR
 	UE_API virtual void PreDuplicateFromRoot(FObjectDuplicationParameters& DupParams) override;
 	UE_API virtual void PreEditChange(FProperty* PropertyAboutToChange) override;
 	UE_API virtual void PostEditChangeChainProperty(struct FPropertyChangedChainEvent& PropertyChangedChainEvent) override;
+	UE_API virtual void PreEditUndo() override;
+	UE_API virtual void PostEditUndo() override;
 	UE_API virtual void PostTransacted(const FTransactionObjectEvent& TransactionEvent) override;
 	UE_API virtual void SetIsTemporarilyHiddenInEditor(bool bIsHidden) override;
 	UE_API virtual bool ShouldTickIfViewportsOnly() const override;
@@ -107,6 +110,17 @@ public:
 	UFUNCTION(BlueprintSetter)
 	UE_API void SetCamera(const FComponentReference& InComponentReference);
 
+	/** Get the composite layers. */
+	UFUNCTION(BlueprintGetter)
+	UE_API TArray<UCompositeLayerBase*> GetCompositeLayers();
+
+	/** Get the composite layers. */
+	UE_API const TArray<TObjectPtr<UCompositeLayerBase>>& GetCompositeLayers() const;
+
+	/** Set the composite layers. */
+	UFUNCTION(BlueprintSetter)
+	UE_API void SetCompositeLayers(TArray<UCompositeLayerBase*> InLayers);
+
 	/** Bind to this to get a callback when the enabled state changes. */
 	FOnCompositeActorEnabledStateChange OnEnabledStateChange;
 
@@ -133,10 +147,10 @@ private:
 	UPROPERTY(EditInstanceOnly, BlueprintGetter = GetCamera, BlueprintSetter = SetCamera, Category = "Composite", meta = (AllowPrivateAccess, UseComponentPicker, AllowAnyActor, AllowedClasses = "/Script/Engine.CameraComponent, /Script/CinematicCamera.CineCameraComponent"))
 	FComponentReference Camera;
 
-public:
-	/** Convenience pre-configured layer of passes. */
-	UPROPERTY(EditAnywhere, BlueprintReadWrite, Instanced, Category = "Composite", meta = (DisplayName = "Layers"))
+	/** Array of composite layers for merging or processing images. */
+	UPROPERTY(EditAnywhere, BlueprintGetter = GetCompositeLayers, BlueprintSetter = SetCompositeLayers, Instanced, Category = "Composite", meta = (DisplayName = "Layers"))
 	TArray<TObjectPtr<UCompositeLayerBase>> CompositeLayers;
+public:
 
 	/** Also provide the post-processing composite graph to SSR input. */
 	UPROPERTY(EditAnywhere, BlueprintReadWrite, Category = "Composite", AdvancedDisplay)
@@ -196,6 +210,7 @@ private:
 	bool bIsModifyingAProperty;
 
 	/** Befriend the actor's editor customizations. */
+	friend class SCompositePanelLayerTree;
 	friend class FCompositeActorCustomization;
 	friend class FCompositeActorPanelDetailCustomization;
 };
diff --git a/Engine/Plugins/Compositing/Composite/Source/Composite/Public/Layers/CompositeLayerBase.h b/Engine/Plugins/Compositing/Composite/Source/Composite/Public/Layers/CompositeLayerBase.h
--- a/Engine/Plugins/Compositing/Composite/Source/Composite/Public/Layers/CompositeLayerBase.h
+++ b/Engine/Plugins/Compositing/Composite/Source/Composite/Public/Layers/CompositeLayerBase.h
@@ -52,6 +52,9 @@ public:
 	UE_API virtual void PostEditChangeProperty(FPropertyChangedEvent& PropertyChangedEvent) override;
 #endif
 
+	/** Called when the layer is removed or composite actor is destroyed. */
+	UE_API virtual void OnRemoved() {}
+
 	/** Getter function to override, returning pass layer proxies to be passed to the render thread. (Proxy objects should be allocated from the provided allocator.) */
 	UE_API virtual bool GetProxy(FTraversalContext& InContext, FSceneRenderingBulkObjectAllocator& InFrameAllocator, FCompositeCorePassProxy*& OutProxy) const { return false; };
 
diff --git a/Engine/Plugins/Compositing/Composite/Source/Composite/Public/Layers/CompositeLayerPlate.h b/Engine/Plugins/Compositing/Composite/Source/Composite/Public/Layers/CompositeLayerPlate.h
--- a/Engine/Plugins/Compositing/Composite/Source/Composite/Public/Layers/CompositeLayerPlate.h
+++ b/Engine/Plugins/Compositing/Composite/Source/Composite/Public/Layers/CompositeLayerPlate.h
@@ -45,6 +45,8 @@ public:
 	/** Destructor. */
 	UE_API virtual ~UCompositeLayerPlate();
 
+	UE_API virtual void OnRemoved() override;
+
 	/** Set the enabled state. */
 	UE_API virtual void SetEnabled(bool bInEnabled) override;
 
diff --git a/Engine/Plugins/Compositing/Composite/Source/Composite/Public/Layers/CompositeLayerSceneCapture.h b/Engine/Plugins/Compositing/Composite/Source/Composite/Public/Layers/CompositeLayerSceneCapture.h
--- a/Engine/Plugins/Compositing/Composite/Source/Composite/Public/Layers/CompositeLayerSceneCapture.h
+++ b/Engine/Plugins/Compositing/Composite/Source/Composite/Public/Layers/CompositeLayerSceneCapture.h
@@ -24,11 +24,17 @@ public:
 	
 	/** Destructor. */
 	UE_API virtual ~UCompositeLayerSceneCapture();
+	
+	UE_API virtual void OnRemoved() override;
+
+	//~ Begin UObject Interface
+	UE_API virtual void BeginDestroy() override;
 
 #if WITH_EDITOR
 	UE_API virtual void PostEditChangeProperty(FPropertyChangedEvent& PropertyChangedEvent) override;
 	UE_API virtual void PostEditUndo() override;
 #endif
+	//~ End UObject Interface
 
 	/** Getter function to override, returning pass layer proxies to be passed to the render thread. (Proxy objects should be allocated from the provided allocator.) */
 	UE_API virtual bool GetProxy(FTraversalContext& InContext, FSceneRenderingBulkObjectAllocator& InFrameAllocator, FCompositeCorePassProxy*& OutProxy) const override;
diff --git a/Engine/Plugins/Compositing/Composite/Source/Composite/Public/Layers/CompositeLayerShadowReflection.h b/Engine/Plugins/Compositing/Composite/Source/Composite/Public/Layers/CompositeLayerShadowReflection.h
--- a/Engine/Plugins/Compositing/Composite/Source/Composite/Public/Layers/CompositeLayerShadowReflection.h
+++ b/Engine/Plugins/Compositing/Composite/Source/Composite/Public/Layers/CompositeLayerShadowReflection.h
@@ -41,6 +41,8 @@ public:
 	/** Destructor. */
 	UE_API virtual ~UCompositeLayerShadowReflection();
 
+	UE_API virtual void OnRemoved() override;
+
 	//~ Begin UObject Interface
 	UE_API virtual void BeginDestroy() override;
 
diff --git a/Engine/Plugins/Compositing/Composite/Source/Composite/Public/Layers/CompositeLayerSingleLightShadow.h b/Engine/Plugins/Compositing/Composite/Source/Composite/Public/Layers/CompositeLayerSingleLightShadow.h
--- a/Engine/Plugins/Compositing/Composite/Source/Composite/Public/Layers/CompositeLayerSingleLightShadow.h
+++ b/Engine/Plugins/Compositing/Composite/Source/Composite/Public/Layers/CompositeLayerSingleLightShadow.h
@@ -37,6 +37,11 @@ public:
 	/** Destructor. */
 	UE_API virtual ~UCompositeLayerSingleLightShadow() = default;
 
+	UE_API virtual void OnRemoved() override;
+
+	//~ Begin UObject Interface
+	UE_API virtual void BeginDestroy() override;
+
 #if WITH_EDITOR
 	UE_API virtual void PostEditChangeProperty(FPropertyChangedEvent& PropertyChangedEvent) override;
 	UE_API virtual void PostEditUndo() override;
diff --git a/Engine/Plugins/Compositing/Composite/Source/CompositeEditor/Private/CompositeEditorModule.cpp b/Engine/Plugins/Compositing/Composite/Source/CompositeEditor/Private/CompositeEditorModule.cpp
--- a/Engine/Plugins/Compositing/Composite/Source/CompositeEditor/Private/CompositeEditorModule.cpp
+++ b/Engine/Plugins/Compositing/Composite/Source/CompositeEditor/Private/CompositeEditorModule.cpp
@@ -5,6 +5,7 @@
 #include "Editor.h"
 #include "EngineUtils.h"
 #include "EngineAnalytics.h"
+#include "HoldoutCompositeComponent.h"
 #include "Framework/Notifications/NotificationManager.h"
 #include "IConcertSyncClient.h"
 #include "IConcertSyncClientModule.h"
@@ -24,6 +25,7 @@
 #include "Layers/CompositeLayerSceneCapture.h"
 #include "Layers/CompositeLayerShadowReflection.h"
 #include "Layers/CompositeLayerSingleLightShadow.h"
+#include "Misc/MessageDialog.h"
 #include "UI/SCompositeEditorPanel.h"
 
 #define LOCTEXT_NAMESPACE "CompositeEditorModule"
@@ -34,6 +36,8 @@ void FCompositeEditorModule::StartupModule()
 {
 	FCoreDelegates::OnPostEngineInit.AddRaw(this, &FCompositeEditorModule::OnPostEngineInit);
 
+	UHoldoutCompositeComponent::OnComponentCreatedDelegate.AddRaw(this, &FCompositeEditorModule::TriggerHoldoutCompositeWarning);
+
 	FCompositeEditorCommands::Register();
 	FCompositeEditorStyle::Register();
 	
@@ -54,10 +58,12 @@ void FCompositeEditorModule::ShutdownModule()
 	FCompositeEditorCommands::Unregister();
 	FCompositeEditorStyle::Unregister();
 	
-	if (GEditor != nullptr)
+	if (GEditor)
 	{
 		GEditor->OnLevelActorAdded().RemoveAll(this);
 	}
+
+	UHoldoutCompositeComponent::OnComponentCreatedDelegate.RemoveAll(this);
 	
 	FCoreDelegates::OnPostEngineInit.RemoveAll(this);
 }
@@ -115,7 +121,11 @@ void FCompositeEditorModule::OnLevelActorAdded(AActor* InActor)
 
 void FCompositeEditorModule::OnConcertSessionStartup(TSharedRef<IConcertClientSession> InSession)
 {
-	const UWorld* World = GEditor->GetEditorWorldContext().World();
+	UWorld* World = nullptr;
+	if (GEditor)
+	{
+		World = GEditor->GetEditorWorldContext().World();
+	}
 
 	if (IsValid(World))
 	{
@@ -130,6 +140,45 @@ void FCompositeEditorModule::OnConcertSessionStartup(TSharedRef<IConcertClientSe
 	}
 }
 
+void FCompositeEditorModule::TriggerHoldoutCompositeWarning(const UHoldoutCompositeComponent* InComponent)
+{
+	static bool bWarnOnce = true;
+
+	if (bWarnOnce)
+	{
+		int32 CompositeActorNum = 0;
+		UWorld* World = nullptr;
+		if (GEditor)
+		{
+			World = GEditor->GetEditorWorldContext().World();
+		}
+
+		if (IsValid(World))
+		{
+			for (TActorIterator<ACompositeActor> It(World); It; ++It)
+			{
+				ACompositeActor* CompositeActor = *It;
+				if (IsValid(CompositeActor))
+				{
+					CompositeActorNum++;
+				}
+			}
+		}
+
+		if (CompositeActorNum > 0)
+		{
+			FNotificationInfo Info(LOCTEXT("FCompositeEditorModuleHoldoutCompositeCreated", "Holdout composite components are not designed to be used with the new Composite Actor. Prefer registering meshes on a plate layer directly."));
+			Info.Image = FAppStyle::GetBrush(TEXT("MessageLog.Warning"));
+			Info.bFireAndForget = true;
+			Info.FadeOutDuration = 4.0f;
+			Info.ExpireDuration = 8.0f;
+			FSlateNotificationManager::Get().AddNotification(Info);
+
+			bWarnOnce = false;
+		}
+	}
+}
+
 void FCompositeEditorModule::RegisterTabSpawners()
 {
 	FLevelEditorModule& LevelEditorModule = FModuleManager::GetModuleChecked<FLevelEditorModule>("LevelEditor");
diff --git a/Engine/Plugins/Compositing/Composite/Source/CompositeEditor/Public/CompositeEditorModule.h b/Engine/Plugins/Compositing/Composite/Source/CompositeEditor/Public/CompositeEditorModule.h
--- a/Engine/Plugins/Compositing/Composite/Source/CompositeEditor/Public/CompositeEditorModule.h
+++ b/Engine/Plugins/Compositing/Composite/Source/CompositeEditor/Public/CompositeEditorModule.h
@@ -30,6 +30,9 @@ private:
 	/** Callback on multi-user concert session start. */
 	void OnConcertSessionStartup(TSharedRef<IConcertClientSession> InSession);
 
+	/** Callback to trigger warnings when an holdout composite component is created with a composite actor. */
+	void TriggerHoldoutCompositeWarning(const class UHoldoutCompositeComponent* InComponent);
+
 	/** Register tab spawners. */
 	void RegisterTabSpawners();
 
diff --git a/Engine/Plugins/Compositing/CompositeCore/Source/CompositeCore/Private/HoldoutCompositeComponent.cpp b/Engine/Plugins/Compositing/CompositeCore/Source/CompositeCore/Private/HoldoutCompositeComponent.cpp
--- a/Engine/Plugins/Compositing/CompositeCore/Source/CompositeCore/Private/HoldoutCompositeComponent.cpp
+++ b/Engine/Plugins/Compositing/CompositeCore/Source/CompositeCore/Private/HoldoutCompositeComponent.cpp
@@ -64,6 +64,8 @@ namespace UE
 	}
 }
 
+FHoldoutCompositeComponentDelegate UHoldoutCompositeComponent::OnComponentCreatedDelegate;
+
 UHoldoutCompositeComponent::UHoldoutCompositeComponent(const FObjectInitializer& ObjectInitializer)
 	: Super(ObjectInitializer)
 {
@@ -143,6 +145,13 @@ void UHoldoutCompositeComponent::PostEditChangeProperty(FPropertyChangedEvent& P
 }
 #endif
 
+void UHoldoutCompositeComponent::OnComponentCreated()
+{
+	Super::OnComponentCreated();
+
+	OnComponentCreatedDelegate.Broadcast(this);
+}
+
 void UHoldoutCompositeComponent::OnAttachmentChanged()
 {
 	Super::OnAttachmentChanged();
diff --git a/Engine/Plugins/Compositing/CompositeCore/Source/CompositeCore/Public/HoldoutCompositeComponent.h b/Engine/Plugins/Compositing/CompositeCore/Source/CompositeCore/Public/HoldoutCompositeComponent.h
--- a/Engine/Plugins/Compositing/CompositeCore/Source/CompositeCore/Public/HoldoutCompositeComponent.h
+++ b/Engine/Plugins/Compositing/CompositeCore/Source/CompositeCore/Public/HoldoutCompositeComponent.h
@@ -7,6 +7,10 @@
 
 #include "HoldoutCompositeComponent.generated.h"
 
+class UHoldoutCompositeComponent;
+
+DECLARE_MULTICAST_DELEGATE_OneParam(FHoldoutCompositeComponentDelegate, const UHoldoutCompositeComponent*);
+
 UCLASS(MinimalAPI, ClassGroup = Rendering, HideCategories=(Activation, Transform, Lighting, Rendering, Tags, Cooking, Physics, LOD, AssetUserData, Navigation), editinlinenew, meta = (BlueprintSpawnableComponent))
 class UHoldoutCompositeComponent : public USceneComponent
 {
@@ -24,6 +28,7 @@ public:
 	//~ End UActorComponent interface
 	
 	//~ Begin USceneComponent Interface.
+	virtual void OnComponentCreated() override;
 	virtual void OnAttachmentChanged() override;
 	virtual void DetachFromComponent(const FDetachmentTransformRules& DetachmentRules) override;
 	//~ End USceneComponent Interface
@@ -40,6 +45,9 @@ public:
 	UFUNCTION(BlueprintSetter)
 	COMPOSITECORE_API void SetEnabled(bool bInEnabled);
 
+	/** Event fired once a component is created. */
+	COMPOSITECORE_API static FHoldoutCompositeComponentDelegate OnComponentCreatedDelegate;
+
 private:
 
 	/* Private implementation of the register method. */
