Composite: Fixed previous CL bughawk-reported issue with toggling solo property on other layers.

#rnx

diff --git a/Engine/Plugins/Compositing/Composite/Source/Composite/Private/Layers/CompositeLayerBase.cpp b/Engine/Plugins/Compositing/Composite/Source/Composite/Private/Layers/CompositeLayerBase.cpp
--- a/Engine/Plugins/Compositing/Composite/Source/Composite/Private/Layers/CompositeLayerBase.cpp
+++ b/Engine/Plugins/Compositing/Composite/Source/Composite/Private/Layers/CompositeLayerBase.cpp
@@ -46,8 +46,9 @@ void UCompositeLayerBase::PostEditChangeProperty(FPropertyChangedEvent& Property
 		{
 			for (const TObjectPtr<UCompositeLayerBase>& CompositeLayer : CompositeActor->GetCompositeLayers())
 			{
-				if (CompositeLayer != this)
+				if (IsValid(CompositeLayer) && CompositeLayer != this)
 				{
+					CompositeLayer->Modify();
 					CompositeLayer->bIsSolo = false;
 				}
 			}
