Composure: General UI improvements, and added icons for Composure objects

diff --git a/Engine/Plugins/Compositing/Composite/Source/Composite/Private/CompositeActor.cpp b/Engine/Plugins/Compositing/Composite/Source/Composite/Private/CompositeActor.cpp
--- a/Engine/Plugins/Compositing/Composite/Source/Composite/Private/CompositeActor.cpp
+++ b/Engine/Plugins/Compositing/Composite/Source/Composite/Private/CompositeActor.cpp
@@ -268,6 +268,14 @@ ACompositeActor::ACompositeActor(const FObjectInitializer& ObjectInitializer)
 	CompositeLayers.Add(CreateDefaultSubobject<UCompositeLayerShadowReflection>(TEXT("ShadowReflectionLayer")));
 	CompositeLayers.Add(CreateDefaultSubobject<UCompositeLayerPlate>(TEXT("PlateLayer")));
 
+#if WITH_EDITOR
+	// Assign default Name value to all layers
+	for (TObjectPtr<UCompositeLayerBase>& Layer : CompositeLayers)
+	{
+		Layer->Name = Layer->GetName();
+	}
+#endif
+	
 	// By default, we disable the shadow reflection catcher due to its significant added cost.
 	CompositeLayers[1]->SetEnabled(false);
 }
diff --git a/Engine/Plugins/Compositing/Composite/Source/CompositeEditor/Private/CompositeEditorStyle.cpp b/Engine/Plugins/Compositing/Composite/Source/CompositeEditor/Private/CompositeEditorStyle.cpp
--- a/Engine/Plugins/Compositing/Composite/Source/CompositeEditor/Private/CompositeEditorStyle.cpp
+++ b/Engine/Plugins/Compositing/Composite/Source/CompositeEditor/Private/CompositeEditorStyle.cpp
@@ -34,6 +34,19 @@ FCompositeEditorStyle::FCompositeEditorStyle()
 	FSlateStyleSet::SetCoreContentRoot(FPaths::EngineContentDir() / TEXT("Editor/Slate/"));
 	
 	Set("CompositeEditor.Composure", new CORE_IMAGE_BRUSH_SVG("Starship/Common/ComposureCompositing", Icon16x16));
+
+	Set("ClassIcon.CompositeActor", new CORE_IMAGE_BRUSH_SVG("Starship/Common/ComposureCompositing", Icon16x16));
+	Set("ClassIcon.CompositeLayerMainRender", new IMAGE_BRUSH_SVG("Editor/Icons/MainRender", Icon16x16));
+	Set("ClassIcon.CompositeLayerShadowReflection", new IMAGE_BRUSH_SVG("Editor/Icons/Shadow", Icon16x16));
+	Set("ClassIcon.CompositeLayerSingleLightShadow", new IMAGE_BRUSH_SVG("Editor/Icons/Shadow", Icon16x16));
+	Set("ClassIcon.CompositeLayerSceneCapture", new CORE_IMAGE_BRUSH_SVG("Starship/AssetIcons/SceneCapture2D_16", Icon16x16));
+	Set("ClassIcon.CompositeLayerPlate", new CORE_IMAGE_BRUSH_SVG("Starship/AssetIcons/MediaPlayer_16", Icon16x16));
+	Set("ClassIcon.CompositePassColorGrade", new CORE_IMAGE_BRUSH_SVG("Starship/TimelineEditor/TrackTypeColor", Icon16x16));
+	Set("ClassIcon.CompositePassColorKeyer", new IMAGE_BRUSH_SVG("Editor/Icons/Chromakey", Icon16x16));
+	Set("ClassIcon.CompositePassFXAA", new CORE_IMAGE_BRUSH_SVG("Starship/Common/AntiAliasing", Icon16x16));
+	Set("ClassIcon.CompositePassSMAA", new CORE_IMAGE_BRUSH_SVG("Starship/Common/AntiAliasing", Icon16x16));
+	Set("ClassIcon.CompositePassOpenColorIO", new IMAGE_BRUSH_SVG("Editor/Icons/OCIO", Icon16x16));
+	Set("ClassIcon.CompositePassMaterial", new CORE_IMAGE_BRUSH_SVG("Starship/AssetIcons/Material_16", Icon16x16));
 	
     Set("CompositeEditor.Passes.Media", new CORE_IMAGE_BRUSH_SVG("Starship/AssetIcons/MediaPlayer_16", Icon16x16));
 	Set("CompositeEditor.Passes.Scene", new CORE_IMAGE_BRUSH_SVG("Starship/Common/Levels", Icon16x16));
diff --git a/Engine/Plugins/Compositing/Composite/Source/CompositeEditor/Private/UI/SCompositeActorPickerTable.cpp b/Engine/Plugins/Compositing/Composite/Source/CompositeEditor/Private/UI/SCompositeActorPickerTable.cpp
--- a/Engine/Plugins/Compositing/Composite/Source/CompositeEditor/Private/UI/SCompositeActorPickerTable.cpp
+++ b/Engine/Plugins/Compositing/Composite/Source/CompositeEditor/Private/UI/SCompositeActorPickerTable.cpp
@@ -435,6 +435,8 @@ void SCompositeActorPickerTable::Construct(const FArguments& InArgs, const FComp
 
 		+SVerticalBox::Slot()
 		.AutoHeight()
+		.MinHeight(24.0f)
+		.VAlign(VAlign_Center)
 		[
 			SNew(STextBlock)
 			.Text(this, &SCompositeActorPickerTable::GetFilterStatusText)
@@ -667,7 +669,7 @@ FText SCompositeActorPickerTable::GetFilterStatusText() const
 		}
 		else
 		{
-			return LOCTEXT("NoActorsFoundLabel", "No actors");
+			return LOCTEXT("NoActorsFoundLabel", "0 Actors");
 		}
 	}
 }
@@ -677,7 +679,7 @@ FSlateColor SCompositeActorPickerTable::GetFilterStatusTextColor() const
 	const int32 NumActors = ActorListItems.Num();
 	const int32 NumFiltered = FilteredActorListItems.Num();
 	
-	if (NumActors == 0 || NumFiltered == 0)
+	if (NumActors > 0 && NumFiltered == 0)
 	{
 		return FAppStyle::Get().GetSlateColor("Colors.AccentRed");
 	}
diff --git a/Engine/Plugins/Compositing/Composite/Source/CompositeEditor/Private/UI/SCompositeEditorPanel.cpp b/Engine/Plugins/Compositing/Composite/Source/CompositeEditor/Private/UI/SCompositeEditorPanel.cpp
--- a/Engine/Plugins/Compositing/Composite/Source/CompositeEditor/Private/UI/SCompositeEditorPanel.cpp
+++ b/Engine/Plugins/Compositing/Composite/Source/CompositeEditor/Private/UI/SCompositeEditorPanel.cpp
@@ -429,8 +429,9 @@ void SCompositeEditorPanel::Construct(const FArguments& InArgs)
 
 		+SSplitter::Slot()
 		.Value(0.25f)
+		.MinSize(this, &SCompositeEditorPanel::GetLayerTreeMinHeight)
 		[
-			SNew(SCompositePanelLayerTree)
+			SAssignNew(LayerTree, SCompositePanelLayerTree)
 			.OnSelectionChanged(this, &SCompositeEditorPanel::OnLayerSelectionChanged)
 		]
 
@@ -450,4 +451,22 @@ void SCompositeEditorPanel::OnLayerSelectionChanged(const TArray<UObject*>& Sele
 	}
 }
 
+float SCompositeEditorPanel::GetLayerTreeMinHeight() const
+{
+	if (LayerTree.IsValid())
+	{
+		const float PanelHeight = GetCachedGeometry().GetAbsoluteSize().Y;
+		const float LayerTreeHeight = LayerTree->GetMinimumHeight();
+
+		// Splitter panels get confused if a slot has a minimum size larger than the splitter, so only
+		// return a non-zero minimum when whole panel is larger than the minimum
+		if (PanelHeight > LayerTreeHeight)
+		{
+			return LayerTreeHeight;
+		}
+	}
+
+	return 0.0f;
+}
+
 #undef LOCTEXT_NAMESPACE
diff --git a/Engine/Plugins/Compositing/Composite/Source/CompositeEditor/Private/UI/SCompositeEditorPanel.h b/Engine/Plugins/Compositing/Composite/Source/CompositeEditor/Private/UI/SCompositeEditorPanel.h
--- a/Engine/Plugins/Compositing/Composite/Source/CompositeEditor/Private/UI/SCompositeEditorPanel.h
+++ b/Engine/Plugins/Compositing/Composite/Source/CompositeEditor/Private/UI/SCompositeEditorPanel.h
@@ -28,6 +28,9 @@ public:
 private:
 	/** Raised when the selected layers in the tree view have changed */
 	void OnLayerSelectionChanged(const TArray<UObject*>& SelectedLayers);
+
+	/** Gets the minimum height to display the full layer tree */
+	float GetLayerTreeMinHeight() const;
 	
 private:
 	/** The tree view widget that displays the composite layers */
diff --git a/Engine/Plugins/Compositing/Composite/Source/CompositeEditor/Private/UI/SCompositePanelLayerTree.cpp b/Engine/Plugins/Compositing/Composite/Source/CompositeEditor/Private/UI/SCompositePanelLayerTree.cpp
--- a/Engine/Plugins/Compositing/Composite/Source/CompositeEditor/Private/UI/SCompositePanelLayerTree.cpp
+++ b/Engine/Plugins/Compositing/Composite/Source/CompositeEditor/Private/UI/SCompositePanelLayerTree.cpp
@@ -349,13 +349,17 @@ public:
 private:
 	virtual TSharedRef<SWidget> GenerateWidgetForColumn(const FName& InColumnName) override
 	{
+		// Padding for column contents for the first column is the base value of 3.5 (based on SHeaderRow)
+		// Subsequent columns need to account for the splitter size (which is default 1.0f)
+		constexpr float FirstColumnHorizontalPadding = 3.5f;
+		constexpr float SplitterSize = 1.0f;
+		constexpr float ColumnHorizontalPadding = FirstColumnHorizontalPadding + SplitterSize;
+		
 		if (InColumnName == CompositePanelLayerTree::LayerTreeColumn_Active)
 		{
-			return SNew(SHorizontalBox)
-
-				+ SHorizontalBox::Slot()
-				.AutoWidth()
-				.Padding(4.f, 0.f, 0.f, 0.f)
+			return SNew(SBorder)
+				.BorderImage(FAppStyle::GetBrush("NoBorder"))
+				.Padding(FirstColumnHorizontalPadding, 0.0f)
 				[
 					SNew(SCheckBox)
 						.Style(FAppStyle::Get(), "RadioButton")
@@ -366,11 +370,9 @@ private:
 		}
 		else if (InColumnName == CompositePanelLayerTree::LayerTreeColumn_Enabled)
 		{
-			return SNew(SHorizontalBox)
-
-				+ SHorizontalBox::Slot()
-				.AutoWidth()
-				.Padding(4.f, 0.f, 0.f, 0.f)
+			return SNew(SBorder)
+				.BorderImage(FAppStyle::GetBrush("NoBorder"))
+				.Padding(ColumnHorizontalPadding, 0.0f)
 				[
 					SNew(SCheckBox)
 						.IsChecked(this, &SCompositePanelLayerTreeItemRow::IsItemEnabled)
@@ -379,9 +381,14 @@ private:
 		}
 		else if (InColumnName == CompositePanelLayerTree::LayerTreeColumn_Solo)
 		{
-			return SNew(SImage).Image(FAppStyle::GetBrush("MediaAsset.AssetActions.Solo.Small"))
-				.ColorAndOpacity_Raw(this, &SCompositePanelLayerTreeItemRow::GetItemSoloIconColor)
-				.OnMouseButtonDown(this, &SCompositePanelLayerTreeItemRow::OnItemSoloIconMouseButtonDown);
+			return SNew(SBorder)
+				.BorderImage(FAppStyle::GetBrush("NoBorder"))
+				.Padding(ColumnHorizontalPadding, 0.0f)
+				[
+					SAssignNew(SoloButtonImage, SImage).Image(FAppStyle::GetBrush("MediaAsset.AssetActions.Solo.Small"))
+					.ColorAndOpacity_Raw(this, &SCompositePanelLayerTreeItemRow::GetItemSoloIconColor)
+					.OnMouseButtonDown(this, &SCompositePanelLayerTreeItemRow::OnItemSoloIconMouseButtonDown)
+				];
 		}
 		else if (InColumnName == CompositePanelLayerTree::LayerTreeColumn_Name)
 		{
@@ -392,7 +399,7 @@ private:
 				.VAlign(VAlign_Center)
 				.Padding(6.f, 0.f, 0.f, 0.f)
 				[
-					SNew(SExpanderArrow, SharedThis(this)).IndentAmount(24)
+					SNew(SExpanderArrow, SharedThis(this)).IndentAmount(12)
 				]
 				
 				+SHorizontalBox::Slot()
@@ -419,11 +426,21 @@ private:
 		}
 		else if (InColumnName == CompositePanelLayerTree::LayerTreeColumn_Type)
 		{
-			return SNew(STextBlock).Text(this, &SCompositePanelLayerTreeItemRow::GetItemType);
+			return SNew(SBorder)
+				.BorderImage(FAppStyle::GetBrush("NoBorder"))
+				.Padding(ColumnHorizontalPadding, 0.0f)
+				[
+					SNew(STextBlock).Text(this, &SCompositePanelLayerTreeItemRow::GetItemType)
+				];
 		}
 		else if (InColumnName == CompositePanelLayerTree::LayerTreeColumn_Operation)
 		{
-			return SNew(STextBlock).Text(this, &SCompositePanelLayerTreeItemRow::GetItemOperation);	
+			return SNew(SBorder)
+				.BorderImage(FAppStyle::GetBrush("NoBorder"))
+				.Padding(ColumnHorizontalPadding, 0.0f)
+				[
+					SNew(STextBlock).Text(this, &SCompositePanelLayerTreeItemRow::GetItemOperation)
+				];
 		}
 		
 		return SNullWidget::NullWidget;
@@ -530,12 +547,12 @@ private:
 			return FLinearColor::Transparent;
 		}
 
-		if (IsHovered())
+		if (SoloButtonImage->IsHovered())
 		{
 			return FStyleColors::ForegroundHover;
 		}
 
-		return bIsLayerSolo ? FSlateColor::UseForeground() : FLinearColor::Transparent;
+		return (bIsLayerSolo || IsHovered()) ? FSlateColor::UseForeground() : FLinearColor::Transparent;
 	}
 
 	/** Raised when the Solo button has been pressed */
@@ -748,6 +765,9 @@ private:
 
 	/** Text block widget used to display the tree item's name */
 	TSharedPtr<SInlineEditableTextBlock> NameTextBlock;
+
+	/** Image widget used for the 'Solo' option button */
+	TSharedPtr<SImage> SoloButtonImage;
 	
 	/** Callback that is raised when the layer's "solo" flag is enabled */
 	FOnLayerSoloToggled OnLayerSoloToggled;
@@ -856,7 +876,9 @@ void SCompositePanelLayerTree::Construct(const FArguments& InArgs)
 				.HAlignCell(HAlign_Center)
 				.VAlignCell(VAlign_Center)
 				[
-					SNew(SImage).Image(FAppStyle::GetBrush("MediaAsset.AssetActions.Solo.Small"))
+					SNew(SImage)
+					.Image(FAppStyle::GetBrush("MediaAsset.AssetActions.Solo.Small"))
+					.ColorAndOpacity(FSlateColor::UseForeground())
 				]
 				
 				+SHeaderRow::Column(CompositePanelLayerTree::LayerTreeColumn_Name)
@@ -897,7 +919,7 @@ void SCompositePanelLayerTree::Construct(const FArguments& InArgs)
 		+SVerticalBox::Slot()
 		.AutoHeight()
 		[
-			SNew(SBorder)
+			SAssignNew(FooterContainer, SBorder)
 			.BorderImage(FAppStyle::Get().GetBrush("Brushes.Header"))
 			.VAlign(VAlign_Center)
 			.HAlign(HAlign_Left)
@@ -956,6 +978,25 @@ void SCompositePanelLayerTree::PostRedo(bool bSuccess)
 	PostUndo(bSuccess);
 }
 
+float SCompositePanelLayerTree::GetMinimumHeight() const
+{
+	// Calculate the number of rows the tree view should be displaying
+	int32 NumRows = 0;
+	for (const FCompositeActorTreeItemPtr& TreeItem : FilteredCompositeActorTreeItems)
+	{
+		NumRows += TreeItem->Children.Num() + 1;
+	}
+	
+	constexpr float HeaderHeight = 24.0f;
+	constexpr float RowHeight = 18.0f;
+
+	const float ToolbarHeight = Toolbar->GetCachedGeometry().GetAbsoluteSize().Y + 12.0f;
+	const float FooterHeight = FooterContainer->GetCachedGeometry().GetAbsoluteSize().Y;
+
+	// Add the toolbar height (plus padding), the tree view header, the tree view rows, and the footer height
+	return ToolbarHeight + HeaderHeight + RowHeight * NumRows + FooterHeight + 2.0f;
+}
+
 void SCompositePanelLayerTree::BindCommands()
 {
 	CommandList->MapAction(FGenericCommands::Get().Copy,
@@ -1551,8 +1592,8 @@ FText SCompositePanelLayerTree::GetFilterStatusText() const
 	TArray<FCompositeActorTreeItemPtr> SelectedItems = TreeView.IsValid() ? TreeView->GetSelectedItems() : TArray<FCompositeActorTreeItemPtr>();
 	const int32 NumSelected = SelectedItems.Num();
 	
-	const FText ActorLabel = NumActors > 1 ? LOCTEXT("ActorPlural", "Actors") : LOCTEXT("ActorSingular", "Actor");
-	if (NumActors == NumFiltered)
+	const FText ActorLabel = NumActors > 1 ? LOCTEXT("ActorPlural", "Composite Actors") : LOCTEXT("ActorSingular", "Composite Actor");
+	if (NumActors > 0 && NumActors == NumFiltered)
 	{
 		if (NumSelected > 0)
 		{
@@ -1589,7 +1630,7 @@ FText SCompositePanelLayerTree::GetFilterStatusText() const
 		}
 		else
 		{
-			return LOCTEXT("NoActorsFoundLabel", "No composite actors found in level");
+			return LOCTEXT("NoActorsFoundLabel", "0 Composite Actors");
 		}
 	}
 }
@@ -1599,7 +1640,7 @@ FSlateColor SCompositePanelLayerTree::GetFilterStatusTextColor() const
 	const int32 NumActors = CompositeActorTreeItems.Num();
 	const int32 NumFiltered = FilteredCompositeActorTreeItems.Num();
 	
-	if (NumActors == 0 || NumFiltered == 0)
+	if (NumActors > 0 && NumFiltered == 0)
 	{
 		return FAppStyle::Get().GetSlateColor("Colors.AccentRed");
 	}
diff --git a/Engine/Plugins/Compositing/Composite/Source/CompositeEditor/Private/UI/SCompositePanelLayerTree.h b/Engine/Plugins/Compositing/Composite/Source/CompositeEditor/Private/UI/SCompositePanelLayerTree.h
--- a/Engine/Plugins/Compositing/Composite/Source/CompositeEditor/Private/UI/SCompositePanelLayerTree.h
+++ b/Engine/Plugins/Compositing/Composite/Source/CompositeEditor/Private/UI/SCompositePanelLayerTree.h
@@ -66,6 +66,9 @@ public:
 	virtual void PostUndo(bool bSuccess) override;
 	virtual void PostRedo(bool bSuccess) override;
 	// End of FEditorUndoClient
+
+	/** Computes the minimum height necessary to display the layer tree with all layers visible */
+	float GetMinimumHeight() const;
 	
 private:	
 	/** Binds commands and callbacks to the command list */
@@ -174,6 +177,9 @@ private:
 	/** Toolbar being displayed above the tree view */
 	TSharedPtr<SCompositePanelLayerTreeToolbar> Toolbar;
 
+	/** The root widget that contains all the footer content */
+	TSharedPtr<SWidget> FooterContainer;
+	
 	/** The command list for commands related to the tree view */
 	TSharedPtr<FUICommandList> CommandList;
 
diff --git a/Engine/Plugins/Compositing/Composite/Source/CompositeEditor/Private/UI/SCompositePassTree.cpp b/Engine/Plugins/Compositing/Composite/Source/CompositeEditor/Private/UI/SCompositePassTree.cpp
--- a/Engine/Plugins/Compositing/Composite/Source/CompositeEditor/Private/UI/SCompositePassTree.cpp
+++ b/Engine/Plugins/Compositing/Composite/Source/CompositeEditor/Private/UI/SCompositePassTree.cpp
@@ -557,9 +557,15 @@ private:
 	{
 		if (InColumnName == CompositePassTree::PassTreeColumn_Enabled)
 		{
-			return SNew(SCheckBox)
-				.IsChecked(this, &SCompositePassTreeItemRow::IsItemEnabled)
-				.OnCheckStateChanged(this, &SCompositePassTreeItemRow::OnIsItemEnabledChanged);
+			return SNew(SBorder)
+				.BorderImage(FAppStyle::GetBrush("NoBorder"))
+				.HAlign(HAlign_Center)
+				.Padding(3.5f, 0.0f)
+				[
+					SNew(SCheckBox)
+					.IsChecked(this, &SCompositePassTreeItemRow::IsItemEnabled)
+					.OnCheckStateChanged(this, &SCompositePassTreeItemRow::OnIsItemEnabledChanged)
+				];
 		}
 		if (InColumnName == CompositePassTree::PassTreeColumn_Name)
 		{
@@ -570,7 +576,7 @@ private:
 				.VAlign(VAlign_Center)
 				.Padding(6.f, 0.f, 0.f, 0.f)
 				[
-					SNew(SExpanderArrow, SharedThis(this)).IndentAmount(24)
+					SNew(SExpanderArrow, SharedThis(this)).IndentAmount(12)
 				]
 				
 				+SHorizontalBox::Slot()
@@ -931,6 +937,8 @@ void SCompositePassTree::Construct(const FArguments& InArgs, const TWeakObjectPt
 
 		+SVerticalBox::Slot()
 		.AutoHeight()
+		.MinHeight(24.0f)
+		.VAlign(VAlign_Center)
 		[
 			SNew(STextBlock)
 			.Text(this, &SCompositePassTree::GetFilterStatusText)
@@ -1454,7 +1462,7 @@ FText SCompositePassTree::GetFilterStatusText() const
 		}
 		else
 		{
-			return LOCTEXT("NoPassesLabel", "No passes on plate");
+			return LOCTEXT("NoPassesLabel", "0 Passes");
 		}
 	}
 }
@@ -1479,7 +1487,7 @@ FSlateColor SCompositePassTree::GetFilterStatusColor() const
 		}
 	}
 
-	if (NumPasses == 0 || NumFilteredPasses == 0)
+	if (NumPasses > 0 && NumFilteredPasses == 0)
 	{
 		return FAppStyle::Get().GetSlateColor("Colors.AccentRed");
 	}
