Composite: Remove single light shadow bias UIMax=1.0 since the value can still often exceed one.

diff --git a/Engine/Plugins/Compositing/Composite/Source/Composite/Public/Layers/CompositeLayerSingleLightShadow.h b/Engine/Plugins/Compositing/Composite/Source/Composite/Public/Layers/CompositeLayerSingleLightShadow.h
--- a/Engine/Plugins/Compositing/Composite/Source/Composite/Public/Layers/CompositeLayerSingleLightShadow.h
+++ b/Engine/Plugins/Compositing/Composite/Source/Composite/Public/Layers/CompositeLayerSingleLightShadow.h
@@ -89,7 +89,7 @@ public:
 	/**
 	* Shadow bias to avoid shadow acne. Increase until unwanted shadowing disappears (automatically scaled by r.Shadow.CSMDepthBias console variable).
 	*/
-	UPROPERTY(EditAnywhere, BlueprintReadWrite, Category = "", meta = (ClampMin = "0.0", UIMin = "0.0", UIMax = "1.0"))
+	UPROPERTY(EditAnywhere, BlueprintReadWrite, Category = "", meta = (ClampMin = "0.0", UIMin = "0.0"))
 	float ShadowBias = 0.5f;
 
 private:
