Composite: Adjust the spawned composite mesh rotation to keep it horizontally level.

diff --git a/Engine/Plugins/Compositing/Composite/Source/CompositeEditor/Private/Customizations/CompositeLayerPlateCustomization.cpp b/Engine/Plugins/Compositing/Composite/Source/CompositeEditor/Private/Customizations/CompositeLayerPlateCustomization.cpp
--- a/Engine/Plugins/Compositing/Composite/Source/CompositeEditor/Private/Customizations/CompositeLayerPlateCustomization.cpp
+++ b/Engine/Plugins/Compositing/Composite/Source/CompositeEditor/Private/Customizations/CompositeLayerPlateCustomization.cpp
@@ -325,7 +325,12 @@ void FCompositeLayerPlateCustomization::CreateCompositeMeshActor()
 		const FVector MeshLoc = CameraLoc - 0.5 * MeshSize * Up - 0.5 * MeshSize * Right + (Distance - MeshSize) * Forward;
 		
 		SpawnTransform.SetLocation(MeshLoc);
-		SpawnTransform.SetRotation(CameraComponent->GetComponentRotation().Quaternion());
+		
+		FRotator Rotation = CameraComponent->GetComponentRotation();
+		// Only preserve the yaw to keep the mesh horizontally level
+		Rotation.Pitch = 0.0f;
+		Rotation.Roll = 0.0f;
+		SpawnTransform.SetRotation(Rotation.Quaternion());
 	}
 
 	FScopedTransaction AddCompositeMeshTransaction(LOCTEXT("AddCompositeMeshTransaction", "Add Composite Mesh"));
