Composure: Added ability to add composite actors from the Composure panel, either when there are no composite actors in the level, or through the Add dropdown menu

diff --git a/Engine/Plugins/Compositing/Composite/Source/CompositeEditor/Private/UI/SCompositePanelLayerTree.cpp b/Engine/Plugins/Compositing/Composite/Source/CompositeEditor/Private/UI/SCompositePanelLayerTree.cpp
--- a/Engine/Plugins/Compositing/Composite/Source/CompositeEditor/Private/UI/SCompositePanelLayerTree.cpp
+++ b/Engine/Plugins/Compositing/Composite/Source/CompositeEditor/Private/UI/SCompositePanelLayerTree.cpp
@@ -55,14 +55,17 @@ public:
 		SLATE_ATTRIBUTE(bool, IsPilotingCamera)
 		SLATE_EVENT(FOnPilotCameraToggled, OnPilotCameraToggled)
 		SLATE_EVENT(FSimpleDelegate, OnFilterChanged)
+		SLATE_EVENT(FSimpleDelegate, OnCompositeActorAdded)
 		SLATE_EVENT(FOnLayerAdded, OnLayerAdded)
 	SLATE_END_ARGS()
 
 	void Construct(const FArguments& InArgs)
 	{
 		IsPilotingCamera = InArgs._IsPilotingCamera;
+		CanAddLayer = InArgs._CanAddLayer;
 		OnPilotCameraToggled = InArgs._OnPilotCameraToggled;
 		OnFilterChanged = InArgs._OnFilterChanged;
+		OnCompositeActorAdded = InArgs._OnCompositeActorAdded;
 		OnLayerAdded = InArgs._OnLayerAdded;
 
 		InitializeFilters();
@@ -84,7 +87,6 @@ public:
 					.Icon(FAppStyle::Get().GetBrush("Icons.Plus"))
 					.Text(LOCTEXT("AddMediaButtonLabel", "Add"))
 					.ToolTipText(LOCTEXT("AddMediaButtonToolTip", "Add new composite layer to the selected composite actor"))
-					.IsEnabled(InArgs._CanAddLayer)
 					.OnGetMenuContent(this, &SCompositePanelLayerTreeToolbar::GetAddMenuContent)
 				]
 
@@ -189,6 +191,18 @@ private:
 	{
 		constexpr bool bCloseMenuAfterSelection = true;
 		FMenuBuilder MenuBuilder(bCloseMenuAfterSelection, nullptr);
+
+		MenuBuilder.BeginSection("Actor", LOCTEXT("AddActorSectionLabel", "Actor"));
+		{
+			MenuBuilder.AddMenuEntry(
+				LOCTEXT("PlaceCompositeActorLabel", "Place Composite Actor"),
+				LOCTEXT("PlaceCompositeActorToolTip", "Add a new Composite Actor to the level"),
+				FSlateIconFinder::FindIconForClass(ACompositeActor::StaticClass()),
+				FUIAction(FExecuteAction::CreateSP(this, &SCompositePanelLayerTreeToolbar::AddActor)));
+		}
+		MenuBuilder.EndSection();
+
+		MenuBuilder.BeginSection("Actor", LOCTEXT("AddLayerSectionLabel", "Layer"));
 		
 		TArray<UClass*> BaseLayerTypes;
 		const bool bRecursive = false;
@@ -214,7 +228,9 @@ private:
 				BaseLayerType->GetDisplayNameText(),
 				TAttribute<FText>(),
 				FSlateIconFinder::FindIconForClass(BaseLayerType),
-				FUIAction(FExecuteAction::CreateSP(this, &SCompositePanelLayerTreeToolbar::AddLayer, BaseLayerType)));
+				FUIAction(
+					FExecuteAction::CreateSP(this, &SCompositePanelLayerTreeToolbar::AddLayer, BaseLayerType),
+					FCanExecuteAction::CreateLambda([this] { return CanAddLayer.Get(true); })));
 		}
 				
 		// Add any blueprint derived classes in a new section of the menu
@@ -241,13 +257,23 @@ private:
 					ChildLayerType->GetDisplayNameText(),
 					TAttribute<FText>(),
 					FSlateIconFinder::FindIconForClass(ChildLayerType),
-					FUIAction(FExecuteAction::CreateSP(this, &SCompositePanelLayerTreeToolbar::AddLayer, ChildLayerType)));
+					FUIAction(
+						FExecuteAction::CreateSP(this, &SCompositePanelLayerTreeToolbar::AddLayer, ChildLayerType),
+						FCanExecuteAction::CreateLambda([this] { return CanAddLayer.Get(true); })));
 			}
 		}
+
+		MenuBuilder.EndSection();
 		
 		return MenuBuilder.MakeWidget();
 	}
 
+	/** Callback raised when 'Place Composite Actor' is selected in the Add dropdown menu */
+	void AddActor()
+	{
+		OnCompositeActorAdded.ExecuteIfBound();
+	}
+	
 	/** Callback raised when a layer type is selected in the Add dropdown menu */
 	void AddLayer(const UClass* InLayerType)
 	{
@@ -264,11 +290,15 @@ private:
 	/** Attribute used to query if a camera actor is actively being piloted in the editor */
 	TAttribute<bool> IsPilotingCamera;
 
+	/** Gets whether there is a valid composite actor that a layer can be added to */
+	TAttribute<bool> CanAddLayer;
+	
 	/** Manages the text filter that has been entered into the filter text box */
 	TSharedPtr<FCompositeLayersFilter> TextFilter;
 
 	FOnPilotCameraToggled OnPilotCameraToggled;
 	FSimpleDelegate OnFilterChanged;
+	FSimpleDelegate OnCompositeActorAdded;
 	FOnLayerAdded OnLayerAdded;
 };
 
@@ -831,89 +861,140 @@ void SCompositePanelLayerTree::Construct(const FArguments& InArgs)
 			.IsPilotingCamera(this, &SCompositePanelLayerTree::IsPilotingCamera)
 			.OnPilotCameraToggled(this, &SCompositePanelLayerTree::OnPilotCameraToggled)
 			.OnFilterChanged(this, &SCompositePanelLayerTree::OnFilterChanged)
+			.OnCompositeActorAdded(this, &SCompositePanelLayerTree::AddCompositeActor)
 			.OnLayerAdded(this, &SCompositePanelLayerTree::OnLayerAdded)
 		]
 
 		+SVerticalBox::Slot()
 		[
-			SAssignNew(TreeView, STreeView<FCompositeActorTreeItemPtr>)
-			.TreeItemsSource(&FilteredCompositeActorTreeItems)
-			.HeaderRow(
-				SNew(SHeaderRow)
-
-				+SHeaderRow::Column(CompositePanelLayerTree::LayerTreeColumn_Active)
-				.DefaultLabel(FText::GetEmpty())
-				.FixedWidth(24.0f)
-				.HAlignHeader(HAlign_Center)
-				.VAlignHeader(VAlign_Center)
-				.HAlignCell(HAlign_Center)
-				.VAlignCell(VAlign_Center)
-				[
-					SNew(SCheckBox)
-						.Style(FAppStyle::Get(), "RadioButton")
-						.IsChecked(this, &SCompositePanelLayerTree::GetGlobalActiveCheckState)
-						.OnCheckStateChanged(this, &SCompositePanelLayerTree::OnGlobalActiveCheckStateChanged)
-				]
-				
-				+SHeaderRow::Column(CompositePanelLayerTree::LayerTreeColumn_Enabled)
-				.DefaultLabel(FText::GetEmpty())
-				.FixedWidth(24.0f)
-				.HAlignHeader(HAlign_Center)
-				.VAlignHeader(VAlign_Center)
-				.HAlignCell(HAlign_Center)
-				.VAlignCell(VAlign_Center)
+			SNew(SWidgetSwitcher)
+			.WidgetIndex_Lambda([this]{ return CompositeActorTreeItems.IsEmpty() ? 0 : 1; })
+
+			+SWidgetSwitcher::Slot()
+			[
+				SNew(SVerticalBox)
+
+				+SVerticalBox::Slot()
+				.AutoHeight()
+				.HAlign(HAlign_Center)
+				.Padding(0.0, 12.0, 0.0, 8.0)
 				[
-					SNew(SCheckBox)
-					.IsChecked(this, &SCompositePanelLayerTree::GetGlobalEnabledCheckState)
-					.OnCheckStateChanged(this, &SCompositePanelLayerTree::OnGlobalEnabledCheckStateChanged)
+					SNew(STextBlock)
+					.Text(LOCTEXT("AddCompositeActorLabel", "Add a Composite Actor"))
+					.Font(FCoreStyle::GetDefaultFontStyle("Italic", 10))
 				]
-				
-				+SHeaderRow::Column(CompositePanelLayerTree::LayerTreeColumn_Solo)
-				.DefaultLabel(FText::GetEmpty())
-				.FixedWidth(24.0f)
-				.HAlignHeader(HAlign_Center)
-				.VAlignHeader(VAlign_Center)
-				.HAlignCell(HAlign_Center)
-				.VAlignCell(VAlign_Center)
+
+				+SVerticalBox::Slot()
+				.AutoHeight()
+				.HAlign(HAlign_Center)
+				.Padding(0.0, 0.0, 12.0, 0.0)
 				[
-					SNew(SImage)
-					.Image(FAppStyle::GetBrush("MediaAsset.AssetActions.Solo.Small"))
-					.ColorAndOpacity(FSlateColor::UseForeground())
+					SNew(SButton)
+					.OnClicked_Lambda([this](){ AddCompositeActor(); return FReply::Handled(); })
+					[
+						SNew(SHorizontalBox)
+						
+						+SHorizontalBox::Slot()
+						.AutoWidth()
+						.Padding(FMargin(0, 1))
+						[
+							SNew(SImage)
+							.Image(FSlateIconFinder::FindIconBrushForClass(ACompositeActor::StaticClass()))
+						]
+						
+						+SHorizontalBox::Slot()
+						.VAlign(VAlign_Center)
+						.AutoWidth()
+						.Padding(FMargin(2, 0, 0, 0))
+						[
+							SNew(STextBlock)
+							.Text(LOCTEXT("PlaceCompositeActorButtonLabel", "Place Composite Actor"))
+						]
+					]
 				]
-				
-				+SHeaderRow::Column(CompositePanelLayerTree::LayerTreeColumn_Name)
-				.DefaultLabel(LOCTEXT("LayerTreeColumnLabel", "Layer"))
-				.FillWidth(0.55)
-
-				+SHeaderRow::Column(CompositePanelLayerTree::LayerTreeColumn_Type)
-				.DefaultLabel(LOCTEXT("TypeTreeColumnLabel", "Type"))
-				.FillWidth(0.3)
-
-				+SHeaderRow::Column(CompositePanelLayerTree::LayerTreeColumn_Operation)
-				.DefaultLabel(LOCTEXT("OperationTreeColumnLabel", "Operation"))
-				.FillWidth(0.15)
-			)
-			.OnGenerateRow_Lambda([this](FCompositeActorTreeItemPtr InTreeItem, const TSharedRef<STableViewBase>& InOwnerTable)
-			{
-				return SNew(SCompositePanelLayerTreeItemRow, InOwnerTable, InTreeItem)
-					.OnLayerSoloToggled(this, &SCompositePanelLayerTree::OnLayerSoloToggled)
-					.OnLayerMoved(this, &SCompositePanelLayerTree::OnLayerMoved);
-			})
-			.OnGetChildren_Lambda([](FCompositeActorTreeItemPtr InTreeItem, TArray<FCompositeActorTreeItemPtr>& OutChildren)
-			{
-				if (InTreeItem.IsValid())
+			]
+
+			+SWidgetSwitcher::Slot()
+			[
+				SAssignNew(TreeView, STreeView<FCompositeActorTreeItemPtr>)
+				.TreeItemsSource(&FilteredCompositeActorTreeItems)
+				.HeaderRow(
+					SNew(SHeaderRow)
+
+					+SHeaderRow::Column(CompositePanelLayerTree::LayerTreeColumn_Active)
+					.DefaultLabel(FText::GetEmpty())
+					.FixedWidth(24.0f)
+					.HAlignHeader(HAlign_Center)
+					.VAlignHeader(VAlign_Center)
+					.HAlignCell(HAlign_Center)
+					.VAlignCell(VAlign_Center)
+					[
+						SNew(SCheckBox)
+							.Style(FAppStyle::Get(), "RadioButton")
+							.IsChecked(this, &SCompositePanelLayerTree::GetGlobalActiveCheckState)
+							.OnCheckStateChanged(this, &SCompositePanelLayerTree::OnGlobalActiveCheckStateChanged)
+					]
+					
+					+SHeaderRow::Column(CompositePanelLayerTree::LayerTreeColumn_Enabled)
+					.DefaultLabel(FText::GetEmpty())
+					.FixedWidth(24.0f)
+					.HAlignHeader(HAlign_Center)
+					.VAlignHeader(VAlign_Center)
+					.HAlignCell(HAlign_Center)
+					.VAlignCell(VAlign_Center)
+					[
+						SNew(SCheckBox)
+						.IsChecked(this, &SCompositePanelLayerTree::GetGlobalEnabledCheckState)
+						.OnCheckStateChanged(this, &SCompositePanelLayerTree::OnGlobalEnabledCheckStateChanged)
+					]
+					
+					+SHeaderRow::Column(CompositePanelLayerTree::LayerTreeColumn_Solo)
+					.DefaultLabel(FText::GetEmpty())
+					.FixedWidth(24.0f)
+					.HAlignHeader(HAlign_Center)
+					.VAlignHeader(VAlign_Center)
+					.HAlignCell(HAlign_Center)
+					.VAlignCell(VAlign_Center)
+					[
+						SNew(SImage)
+						.Image(FAppStyle::GetBrush("MediaAsset.AssetActions.Solo.Small"))
+						.ColorAndOpacity(FSlateColor::UseForeground())
+					]
+					
+					+SHeaderRow::Column(CompositePanelLayerTree::LayerTreeColumn_Name)
+					.DefaultLabel(LOCTEXT("LayerTreeColumnLabel", "Layer"))
+					.FillWidth(0.55)
+
+					+SHeaderRow::Column(CompositePanelLayerTree::LayerTreeColumn_Type)
+					.DefaultLabel(LOCTEXT("TypeTreeColumnLabel", "Type"))
+					.FillWidth(0.3)
+
+					+SHeaderRow::Column(CompositePanelLayerTree::LayerTreeColumn_Operation)
+					.DefaultLabel(LOCTEXT("OperationTreeColumnLabel", "Operation"))
+					.FillWidth(0.15)
+				)
+				.OnGenerateRow_Lambda([this](FCompositeActorTreeItemPtr InTreeItem, const TSharedRef<STableViewBase>& InOwnerTable)
+				{
+					return SNew(SCompositePanelLayerTreeItemRow, InOwnerTable, InTreeItem)
+						.OnLayerSoloToggled(this, &SCompositePanelLayerTree::OnLayerSoloToggled)
+						.OnLayerMoved(this, &SCompositePanelLayerTree::OnLayerMoved);
+				})
+				.OnGetChildren_Lambda([](FCompositeActorTreeItemPtr InTreeItem, TArray<FCompositeActorTreeItemPtr>& OutChildren)
 				{
-					for (FCompositeActorTreeItemPtr& Child : InTreeItem->Children)
+					if (InTreeItem.IsValid())
 					{
-						if (Child.IsValid() && !Child->bFilteredOut)
+						for (FCompositeActorTreeItemPtr& Child : InTreeItem->Children)
 						{
-							OutChildren.Add(Child);
+							if (Child.IsValid() && !Child->bFilteredOut)
+							{
+								OutChildren.Add(Child);
+							}
 						}
 					}
-				}
-			})
-			.OnSelectionChanged(this, &SCompositePanelLayerTree::OnLayerSelectionChanged)
-			.OnContextMenuOpening(this, &SCompositePanelLayerTree::CreateTreeContextMenu)
+				})
+				.OnSelectionChanged(this, &SCompositePanelLayerTree::OnLayerSelectionChanged)
+				.OnContextMenuOpening(this, &SCompositePanelLayerTree::CreateTreeContextMenu)
+			]
 		]
 
 		+SVerticalBox::Slot()
@@ -992,9 +1073,10 @@ float SCompositePanelLayerTree::GetMinimumHeight() const
 
 	const float ToolbarHeight = Toolbar->GetCachedGeometry().GetAbsoluteSize().Y + 12.0f;
 	const float FooterHeight = FooterContainer->GetCachedGeometry().GetAbsoluteSize().Y;
+	const float TreeHeight = NumRows > 0 ? RowHeight * NumRows : 48.0f;
 
 	// Add the toolbar height (plus padding), the tree view header, the tree view rows, and the footer height
-	return ToolbarHeight + HeaderHeight + RowHeight * NumRows + FooterHeight + 2.0f;
+	return ToolbarHeight + HeaderHeight + TreeHeight + FooterHeight + 2.0f;
 }
 
 void SCompositePanelLayerTree::BindCommands()
@@ -1678,6 +1760,39 @@ TSharedPtr<SWidget> SCompositePanelLayerTree::CreateTreeContextMenu()
 	return SNullWidget::NullWidget;
 }
 
+void SCompositePanelLayerTree::AddCompositeActor()
+{
+	UWorld* World = GEditor ? GEditor->GetEditorWorldContext().World() : nullptr;
+	if (!World)
+	{
+		return;
+	}
+
+	FScopedTransaction AddCompositeActorTransaction(LOCTEXT("AddCompositeActorTransaction", "Add Composite Actor"));
+	
+	ACompositeActor* AddedActor = Cast<ACompositeActor>(World->SpawnActor(ACompositeActor::StaticClass()));
+	if (AddedActor)
+	{
+		if (GEditor)
+		{
+			GEditor->SelectNone(true, true);
+			GEditor->SelectActor(AddedActor, true, true);
+		}
+
+		RefreshAndExpandTreeView();
+		
+		FCompositeActorTreeItemPtr* TreeItem = CompositeActorTreeItems.FindByPredicate([AddedActor](const FCompositeActorTreeItemPtr& InTreeItem)
+		{
+			return InTreeItem.IsValid() && InTreeItem->CompositeActor == AddedActor;
+		});
+
+		if (TreeItem && TreeView.IsValid())
+		{
+			TreeView->SetSelection(*TreeItem);
+		}
+	}
+}
+
 void SCompositePanelLayerTree::OnLayerAdded(const UClass* InLayerClass)
 {
 	TWeakObjectPtr<ACompositeActor> CompositeActorWeak = nullptr;
diff --git a/Engine/Plugins/Compositing/Composite/Source/CompositeEditor/Private/UI/SCompositePanelLayerTree.h b/Engine/Plugins/Compositing/Composite/Source/CompositeEditor/Private/UI/SCompositePanelLayerTree.h
--- a/Engine/Plugins/Compositing/Composite/Source/CompositeEditor/Private/UI/SCompositePanelLayerTree.h
+++ b/Engine/Plugins/Compositing/Composite/Source/CompositeEditor/Private/UI/SCompositePanelLayerTree.h
@@ -128,6 +128,9 @@ private:
 	/** Creates the right click context menu for the tree view */
 	TSharedPtr<SWidget> CreateTreeContextMenu();
 
+	/** Adds a new composite actor to the level */
+	void AddCompositeActor();
+	
 	/** Raised when a new layer is being added to the selected composite actor */
 	void OnLayerAdded(const UClass* InLayerClass);
 
