Composite: Added convenience for setting materials on plate meshes.

diff --git a/Engine/Plugins/Compositing/Composite/Source/CompositeEditor/Private/Customizations/CompositeLayerPlateCustomization.cpp b/Engine/Plugins/Compositing/Composite/Source/CompositeEditor/Private/Customizations/CompositeLayerPlateCustomization.cpp
--- a/Engine/Plugins/Compositing/Composite/Source/CompositeEditor/Private/Customizations/CompositeLayerPlateCustomization.cpp
+++ b/Engine/Plugins/Compositing/Composite/Source/CompositeEditor/Private/Customizations/CompositeLayerPlateCustomization.cpp
@@ -97,6 +97,7 @@ void FCompositeLayerPlateCustomization::CustomizeDetails(IDetailLayoutBuilder& D
 			SNew(SCompositeActorPickerTable, ActorListRef)
 			.OnExtendAddMenu(this, &FCompositeLayerPlateCustomization::OnExtendCompositeMeshAddMenu)
 			.OnLayoutSizeChanged(this, &FCompositeLayerPlateCustomization::OnLayoutSizeChanged)
+			.ShowApplyMaterialSection(true)
 		];
 		
 		IDetailGroup& PassesGroup = PlateCategory.AddGroup("Passes", LOCTEXT("PassesGroupName", "Passes"), false, true);
diff --git a/Engine/Plugins/Compositing/Composite/Source/CompositeEditor/Private/UI/SCompositeActorPickerTable.cpp b/Engine/Plugins/Compositing/Composite/Source/CompositeEditor/Private/UI/SCompositeActorPickerTable.cpp
--- a/Engine/Plugins/Compositing/Composite/Source/CompositeEditor/Private/UI/SCompositeActorPickerTable.cpp
+++ b/Engine/Plugins/Compositing/Composite/Source/CompositeEditor/Private/UI/SCompositeActorPickerTable.cpp
@@ -2,6 +2,7 @@
 
 #include "SCompositeActorPickerTable.h"
 
+#include "Components/PrimitiveComponent.h"
 #include "CompositeEditorCommands.h"
 #include "Editor.h"
 #include "EditorActorFolders.h"
@@ -15,6 +16,7 @@
 #include "Editor/EditorWidgets/Public/SDropTarget.h"
 #include "Filters/SFilterSearchBox.h"
 #include "Framework/MultiBox/MultiBoxBuilder.h"
+#include "Materials/Material.h"
 #include "Misc/TextFilter.h"
 #include "Modules/ModuleManager.h"
 #include "Styling/SlateIconFinder.h"
@@ -376,6 +378,7 @@ void SCompositeActorPickerTable::Construct(const FArguments& InArgs, const FComp
 {
 	ActorListRef = InActorListRef;
 	OnLayoutSizeChanged = InArgs._OnLayoutSizeChanged;
+	bShowApplyMaterialSection = InArgs._ShowApplyMaterialSection;
 	
 	FCoreUObjectDelegates::OnObjectPropertyChanged.AddSP(this, &SCompositeActorPickerTable::OnObjectPropertyChanged);
 	
@@ -496,7 +499,7 @@ void SCompositeActorPickerTable::BindCommands()
 {
 	CommandList->MapAction(FCompositeEditorCommands::Get().RemoveActor,
 		FExecuteAction::CreateSP(this, &SCompositeActorPickerTable::RemoveActors),
-		FCanExecuteAction::CreateSP(this, &SCompositeActorPickerTable::CanRemoveActors));
+		FCanExecuteAction::CreateSP(this, &SCompositeActorPickerTable::CanEditActors));
 }
 
 void SCompositeActorPickerTable::FillActorList()
@@ -558,7 +561,35 @@ TSharedPtr<SWidget> SCompositeActorPickerTable::CreateListContextMenu()
 		constexpr bool bShouldCloseWindowAfterMenuSelection = true;
 		FMenuBuilder MenuBuilder(bShouldCloseWindowAfterMenuSelection, CommandList);
 		
+		if (bShowApplyMaterialSection.Get())
+		{
+			MenuBuilder.BeginSection("ApplyMaterial", LOCTEXT("ApplyMaterial", "Apply Material"));
+
+			MenuBuilder.AddMenuEntry(
+				LOCTEXT("ApplyLitMaskedMaterial", "Apply Lit Masked Material To Selected Actors"),
+				LOCTEXT("ApplyLitMaskedMaterialToolTip", "Apply the plugin default lit masked material to selected actors: best for catching shadows and reflections."),
+				FSlateIcon(),
+				FUIAction(
+					FExecuteAction::CreateSP(this, &SCompositeActorPickerTable::ApplyMaterialToActors, TEXT("/Composite/Materials/M_CompositeMesh_Lit_Masked.M_CompositeMesh_Lit_Masked")),
+					FCanExecuteAction::CreateSP(this, &SCompositeActorPickerTable::CanEditActors)
+				)
+			);
+
+			MenuBuilder.AddMenuEntry(
+				LOCTEXT("ApplyUnlitAlphaMaterial", "Apply Unlit Alpha Material To Selected Actors"),
+				LOCTEXT("ApplyUnlitAlphaMaterialToolTip", "Apply the plugin default unlit alpha material to selected actors: best for keying media."),
+				FSlateIcon(),
+				FUIAction(
+					FExecuteAction::CreateSP(this, &SCompositeActorPickerTable::ApplyMaterialToActors, TEXT("/Composite/Materials/M_CompositeMesh_Unlit_AlphaComposite.M_CompositeMesh_Unlit_AlphaComposite")),
+					FCanExecuteAction::CreateSP(this, &SCompositeActorPickerTable::CanEditActors)
+				)
+			);
+			MenuBuilder.EndSection();
+		}
+		
+		MenuBuilder.BeginSection("Edit", LOCTEXT("Edit", "Edit"));
 		MenuBuilder.AddMenuEntry(FCompositeEditorCommands::Get().RemoveActor);
+		MenuBuilder.EndSection();
 		
 		return MenuBuilder.MakeWidget();
 	}
@@ -734,7 +765,70 @@ void SCompositeActorPickerTable::RemoveActors()
 	FillActorList();
 }
 
-bool SCompositeActorPickerTable::CanRemoveActors() const
+void SCompositeActorPickerTable::ApplyMaterialToActors(const TCHAR* InMaterialPath)
+{
+	UMaterial* LoadedMaterial = LoadObject<UMaterial>(nullptr, InMaterialPath);
+	if (!ensureMsgf(LoadedMaterial, TEXT("Expected valid material to apply to actors.")))
+	{
+		return;
+	}
+
+	TArray<FActorListItemRef> SelectedActors = ListView.IsValid() ? ListView->GetSelectedItems() : TArray<FActorListItemRef>();
+	if (SelectedActors.IsEmpty())
+	{
+		return;
+	}
+
+	if (!ActorListRef.IsValid())
+	{
+		return;
+	}
+
+	TArray<int32> ActorIndicesToEdit;
+	for (const FActorListItemRef& ActorItem : SelectedActors)
+	{
+		if (!ActorListRef.ActorList->IsValidIndex(ActorItem->Index))
+		{
+			return;
+		}
+
+		ActorIndicesToEdit.Add(ActorItem->Index);
+	}
+
+	if (ActorIndicesToEdit.IsEmpty())
+	{
+		return;
+	}
+
+	TStrongObjectPtr<UObject> PinnedListOwner = ActorListRef.ActorListOwner.Pin();
+
+	FScopedTransaction UpdateActorsMaterialTransaction(LOCTEXT("UpdateActorsMaterial", "Update Actors Material"));
+	PinnedListOwner->Modify();
+	ActorListRef.NotifyPreEditChange();
+
+	for (int32 Index : ActorIndicesToEdit)
+	{
+		for (TSoftObjectPtr<AActor>& Actor : *ActorListRef.ActorList)
+		{
+			for (UActorComponent* Component : Actor->GetComponents())
+			{
+				if (UPrimitiveComponent* PrimitiveComponent = Cast<UPrimitiveComponent>(Component))
+				{
+					PrimitiveComponent->Modify();
+					const int32 MaterialCount = PrimitiveComponent->GetNumMaterials();
+					for (int32 MaterialIndex = 0; MaterialIndex < MaterialCount; MaterialIndex++)
+					{
+						PrimitiveComponent->SetMaterial(MaterialIndex, LoadedMaterial);
+					}
+				}
+			}
+		}
+	}
+	// Rely on the global UpdateCompositeMeshes() post-edit call to update materials into MIDs
+	ActorListRef.NotifyPostEditChangeList(EPropertyChangeType::Unspecified);
+}
+
+bool SCompositeActorPickerTable::CanEditActors() const
 {
 	const int32 NumSelected = ListView.IsValid() ? ListView->GetNumItemsSelected() : 0;
 	if (NumSelected == 0)
diff --git a/Engine/Plugins/Compositing/Composite/Source/CompositeEditor/Private/UI/SCompositeActorPickerTable.h b/Engine/Plugins/Compositing/Composite/Source/CompositeEditor/Private/UI/SCompositeActorPickerTable.h
--- a/Engine/Plugins/Compositing/Composite/Source/CompositeEditor/Private/UI/SCompositeActorPickerTable.h
+++ b/Engine/Plugins/Compositing/Composite/Source/CompositeEditor/Private/UI/SCompositeActorPickerTable.h
@@ -67,10 +67,11 @@ public:
 	DECLARE_DELEGATE_OneParam(FOnExtendAddMenu, FMenuBuilder&);
 	
 public:
-	SLATE_BEGIN_ARGS(SCompositeActorPickerTable) { }
+	SLATE_BEGIN_ARGS(SCompositeActorPickerTable) : _ShowApplyMaterialSection(false) { }
 		SLATE_ATTRIBUTE(TSharedPtr<FSceneOutlinerFilters>, SceneOutlinerFilters)
 		SLATE_EVENT(FOnExtendAddMenu, OnExtendAddMenu)
 		SLATE_EVENT(FSimpleDelegate, OnLayoutSizeChanged)
+		SLATE_ATTRIBUTE(bool, ShowApplyMaterialSection)
 	SLATE_END_ARGS()
 
 	void Construct(const FArguments& InArgs, const FCompositeActorPickerListRef& InActorListRef);
@@ -118,8 +119,11 @@ private:
 	/** Removes the selected actors from the list */
 	void RemoveActors();
 
-	/** Gets whether the RemoveActors command can be used */
-	bool CanRemoveActors() const;
+	/** Apply the specified material to the the selected actors from the list */
+	void ApplyMaterialToActors(const TCHAR* InMaterialPath);
+
+	/** Gets whether the contextual menu commands can be used */
+	bool CanEditActors() const;
 	
 	/** Raised when a property on the object has changed */
 	void OnObjectPropertyChanged(UObject* InObject, FPropertyChangedEvent& InPropertyChangedEvent);
@@ -152,6 +156,9 @@ private:
 	/** Raised when the number of items in the list view changes */
 	FSimpleDelegate OnLayoutSizeChanged;
 
+	/** Flag to enable the material application section for selected actors. */
+	TAttribute<bool> bShowApplyMaterialSection;
+
 	friend class SCompositeActorPickerToolbar;
 	friend class SCompositeActorListItemRow;
 };
