Composure: Fixed issues where widgets in the Composure panel would refresh on Undo/Redo callbacks unnecessarily.

diff --git a/Engine/Plugins/Compositing/Composite/Source/CompositeEditor/Private/UI/SCompositeActorPickerTable.cpp b/Engine/Plugins/Compositing/Composite/Source/CompositeEditor/Private/UI/SCompositeActorPickerTable.cpp
--- a/Engine/Plugins/Compositing/Composite/Source/CompositeEditor/Private/UI/SCompositeActorPickerTable.cpp
+++ b/Engine/Plugins/Compositing/Composite/Source/CompositeEditor/Private/UI/SCompositeActorPickerTable.cpp
@@ -504,13 +504,43 @@ void SCompositeActorPickerTable::BindCommands()
 
 void SCompositeActorPickerTable::FillActorList()
 {
-	ActorListItems.Empty();
-	
 	if (!ActorListRef.IsValid())
 	{
 		return;
 	}
 
+	// Check to see if a refresh is needed by comparing the list of actors being displayed to
+	// the list of actors in the actor list reference
+	bool bNeedsRefresh = false;
+	if (ActorListItems.Num() != ActorListRef.ActorList->Num())
+	{
+		bNeedsRefresh = true;
+	}
+	else
+	{
+		for (const FActorListItemRef& ActorListItem : ActorListItems)
+		{
+			if (!ActorListItem.IsValid())
+			{
+				bNeedsRefresh = true;
+				break;
+			}
+			
+			if (!ActorListRef.ActorList->Contains(ActorListItem->Actor))
+			{
+				bNeedsRefresh = true;
+				break;
+			}
+		}
+	}
+
+	if (!bNeedsRefresh)
+	{
+		return;
+	}
+
+	ActorListItems.Empty();
+	
 	for (int32 Index = 0; Index < ActorListRef.ActorList->Num(); ++Index)
 	{
 		ActorListItems.Add(MakeShared<FActorListItem>((*ActorListRef.ActorList)[Index], Index));
@@ -808,18 +838,16 @@ void SCompositeActorPickerTable::ApplyMaterialToActors(const TCHAR* InMaterialPa
 
 	for (int32 Index : ActorIndicesToEdit)
 	{
-		for (TSoftObjectPtr<AActor>& Actor : *ActorListRef.ActorList)
+		TSoftObjectPtr<AActor>& Actor = (*ActorListRef.ActorList)[Index];
+		for (UActorComponent* Component : Actor->GetComponents())
 		{
-			for (UActorComponent* Component : Actor->GetComponents())
+			if (UPrimitiveComponent* PrimitiveComponent = Cast<UPrimitiveComponent>(Component))
 			{
-				if (UPrimitiveComponent* PrimitiveComponent = Cast<UPrimitiveComponent>(Component))
+				PrimitiveComponent->Modify();
+				const int32 MaterialCount = PrimitiveComponent->GetNumMaterials();
+				for (int32 MaterialIndex = 0; MaterialIndex < MaterialCount; MaterialIndex++)
 				{
-					PrimitiveComponent->Modify();
-					const int32 MaterialCount = PrimitiveComponent->GetNumMaterials();
-					for (int32 MaterialIndex = 0; MaterialIndex < MaterialCount; MaterialIndex++)
-					{
-						PrimitiveComponent->SetMaterial(MaterialIndex, LoadedMaterial);
-					}
+					PrimitiveComponent->SetMaterial(MaterialIndex, LoadedMaterial);
 				}
 			}
 		}
diff --git a/Engine/Plugins/Compositing/Composite/Source/CompositeEditor/Private/UI/SCompositePanelLayerTree.cpp b/Engine/Plugins/Compositing/Composite/Source/CompositeEditor/Private/UI/SCompositePanelLayerTree.cpp
--- a/Engine/Plugins/Compositing/Composite/Source/CompositeEditor/Private/UI/SCompositePanelLayerTree.cpp
+++ b/Engine/Plugins/Compositing/Composite/Source/CompositeEditor/Private/UI/SCompositePanelLayerTree.cpp
@@ -1166,31 +1166,52 @@ void SCompositePanelLayerTree::FillCompositeActorTree(bool bPreserveSelection)
 
 	if (TreeView.IsValid())
 	{
-		for (const TPair<ACompositeActor*, int32>& Pair : ItemsToReselect)
+		// Silence any selection changes while re-selecting, and track if the selection has changed (i.e. the originally selected items no longer exist)
+		// If the selection has changed at the end of re-selection, the selection changed delegate will be invoked manually
+		bool bSelectionChanged = false;
 		{
-			ACompositeActor* CompositeActorToReselect = Pair.Key;
-			int32 LayerToReselect = Pair.Value;
-
-			const int32 CompositeActorItemIndex = FilteredCompositeActorTreeItems.IndexOfByPredicate([CompositeActorToReselect](const FCompositeActorTreeItemPtr& InItem)
+			TGuardValue<bool> SilenceSelectionChangesGuard(bSilenceSelectionChanges, true);
+		
+			for (const TPair<ACompositeActor*, int32>& Pair : ItemsToReselect)
 			{
-				return InItem->CompositeActor == CompositeActorToReselect;
-			});
+				ACompositeActor* CompositeActorToReselect = Pair.Key;
+				int32 LayerToReselect = Pair.Value;
 
-			if (CompositeActorItemIndex == INDEX_NONE)
-			{
-				continue;
-			}
+				const int32 CompositeActorItemIndex = FilteredCompositeActorTreeItems.IndexOfByPredicate([CompositeActorToReselect](const FCompositeActorTreeItemPtr& InItem)
+				{
+					return InItem->CompositeActor == CompositeActorToReselect;
+				});
 
-			FCompositeActorTreeItemPtr CompositeActorItem = FilteredCompositeActorTreeItems[CompositeActorItemIndex];
-			if (LayerToReselect == CompositeActorItem->LayerIndex)
-			{
-				TreeView->SetItemSelection(CompositeActorItem, true);
-			}
-			else if (CompositeActorItem->Children.IsValidIndex(LayerToReselect) && !CompositeActorItem->Children[LayerToReselect]->bFilteredOut)
-			{
-				TreeView->SetItemSelection(CompositeActorItem->Children[LayerToReselect], true);
+				if (CompositeActorItemIndex == INDEX_NONE)
+				{
+					// Composite actor that was previously selected is no longer in the list, likely from deletion, so
+					// we must invoke SelectionChanged to notify any external widgets (e.g. details panel) to clear and refresh
+					bSelectionChanged = true;
+					continue;
+				}
+
+				FCompositeActorTreeItemPtr CompositeActorItem = FilteredCompositeActorTreeItems[CompositeActorItemIndex];
+				if (LayerToReselect == CompositeActorItem->LayerIndex)
+				{
+					TreeView->SetItemSelection(CompositeActorItem, true);
+				}
+				else if (CompositeActorItem->Children.IsValidIndex(LayerToReselect) && !CompositeActorItem->Children[LayerToReselect]->bFilteredOut)
+				{
+					TreeView->SetItemSelection(CompositeActorItem->Children[LayerToReselect], true);
+				}
+				else
+				{
+					// Index of composite layer that was previously selected is no longer valid, indicating the layer was deleted, so
+					// we must invoke SelectionChanged to notify any external widgets (e.g. details panel) to clear and refresh
+					bSelectionChanged = true;
+				}
 			}
 		}
+
+		if (bSelectionChanged)
+		{
+			OnLayerSelectionChanged(nullptr, ESelectInfo::Direct);
+		}
 	}
 }
 
@@ -1611,6 +1632,11 @@ void SCompositePanelLayerTree::OnGlobalEnabledCheckStateChanged(ECheckBoxState C
 
 void SCompositePanelLayerTree::OnLayerSelectionChanged(FCompositeActorTreeItemPtr InTreeItem, ESelectInfo::Type SelectInfo)
 {
+	if (bSilenceSelectionChanges)
+	{
+		return;
+	}
+	
 	TArray<FCompositeActorTreeItemPtr> SelectedItems = TreeView.IsValid() ? TreeView->GetSelectedItems() : TArray<FCompositeActorTreeItemPtr>();
 	
 	TArray<UObject*> SelectedLayers;
diff --git a/Engine/Plugins/Compositing/Composite/Source/CompositeEditor/Private/UI/SCompositePanelLayerTree.h b/Engine/Plugins/Compositing/Composite/Source/CompositeEditor/Private/UI/SCompositePanelLayerTree.h
--- a/Engine/Plugins/Compositing/Composite/Source/CompositeEditor/Private/UI/SCompositePanelLayerTree.h
+++ b/Engine/Plugins/Compositing/Composite/Source/CompositeEditor/Private/UI/SCompositePanelLayerTree.h
@@ -200,6 +200,9 @@ private:
 
 	/** Cached field of view to re-apply once piloting ends. */
 	TOptional<float> CachedViewFOV;
+
+	/** When true, selection changes in the tree view will not invoke the SelectionChanged delegate */
+	bool bSilenceSelectionChanges = false;
 	
 	friend class SCompositePanelLayerTreeItemRow;
 	friend class SCompositePanelLayerTreeToolbar;
