Composite: Meshes registered on a scene capture layer will no longer be visible in the main render by default.

* Users still have the a convenience toggle to disable behavior for cases like masking where the main render should still show the meshes.

diff --git a/Engine/Plugins/Compositing/Composite/Source/Composite/Private/Layers/CompositeLayerSceneCapture.cpp b/Engine/Plugins/Compositing/Composite/Source/Composite/Private/Layers/CompositeLayerSceneCapture.cpp
--- a/Engine/Plugins/Compositing/Composite/Source/Composite/Private/Layers/CompositeLayerSceneCapture.cpp
+++ b/Engine/Plugins/Compositing/Composite/Source/Composite/Private/Layers/CompositeLayerSceneCapture.cpp
@@ -11,9 +11,10 @@
 
 UCompositeLayerSceneCapture::UCompositeLayerSceneCapture(const FObjectInitializer& ObjectInitializer)
 	: Super(ObjectInitializer)
+	, bVisibleInSceneCaptureOnly{true} // By default, registered meshes will no longer be visible in the main render
+	, bCustomRenderPass{false}
 {
 	Operation = ECompositeCoreMergeOp::Over;
-	bCustomRenderPass = false;
 }
 
 UCompositeLayerSceneCapture::~UCompositeLayerSceneCapture() = default;
@@ -26,6 +27,8 @@ void UCompositeLayerSceneCapture::OnRemoved()
 	{
 		CompositeActor->DestroySceneCaptures(this);
 	}
+
+	RestorePrimitiveVisibilityState();
 }
 
 void UCompositeLayerSceneCapture::BeginDestroy()
@@ -50,6 +53,10 @@ void UCompositeLayerSceneCapture::PostEditChangeProperty(FPropertyChangedEvent&
 	{
 		SetCustomRenderPass(bCustomRenderPass);
 	}
+	else if (PropertyName == GET_MEMBER_NAME_CHECKED(ThisClass, bVisibleInSceneCaptureOnly))
+	{
+		SetVisibileInSceneCaptureOnly(bVisibleInSceneCaptureOnly);
+	}
 }
 
 void UCompositeLayerSceneCapture::PostEditUndo()
@@ -108,33 +115,39 @@ const TArray<TSoftObjectPtr<AActor>> UCompositeLayerSceneCapture::GetActors() co
 
 void UCompositeLayerSceneCapture::SetActors(TArray<TSoftObjectPtr<AActor>> InActors)
 {
-	USceneCaptureComponent2D* SceneCaptureComponent = GetSceneCaptureComponent();
-	if (IsValid(SceneCaptureComponent))
-	{
-		// First, we empty show-only actors since our previous logic used them
-		SceneCaptureComponent->ShowOnlyActors.Empty();
-
-		TArray<TWeakObjectPtr<UPrimitiveComponent>> ShowOnlyComponents;
-		ShowOnlyComponents.Reserve(InActors.Num());
+	TArray<TWeakObjectPtr<UPrimitiveComponent>> ShowOnlyComponents;
+	ShowOnlyComponents.Reserve(InActors.Num());
 
-		for (const TSoftObjectPtr<AActor>& SoftActor : InActors)
+	for (const TSoftObjectPtr<AActor>& SoftActor : InActors)
+	{
+		if (const AActor* Actor = SoftActor.Get())
 		{
-			if (const AActor* Actor = SoftActor.Get())
+			for (UActorComponent* Component : Actor->GetComponents())
 			{
-				for (UActorComponent* Component : Actor->GetComponents())
+				if (UPrimitiveComponent* PrimitiveComponent = Cast<UPrimitiveComponent>(Component))
 				{
-					if (UPrimitiveComponent* PrimitiveComponent = Cast<UPrimitiveComponent>(Component))
-					{
-						ShowOnlyComponents.Add(PrimitiveComponent);
-					}
+					ShowOnlyComponents.Add(PrimitiveComponent);
 				}
 			}
 		}
+	}
 
+	USceneCaptureComponent2D* SceneCaptureComponent = GetSceneCaptureComponent();
+	if (IsValid(SceneCaptureComponent))
+	{
+		// First, we empty show-only actors since our previous logic used them
+		SceneCaptureComponent->ShowOnlyActors.Empty();
 		SceneCaptureComponent->ShowOnlyComponents = ShowOnlyComponents;
 	}
 
+	// Note: Is it important to restore the state first since this may be called during undo/redo
+	RestorePrimitiveVisibilityState();
+
 	Actors = MoveTemp(InActors);
+
+	CachePrimitiveVisibilityState(ShowOnlyComponents);
+
+	UpdatePrimitiveVisibilityState();
 }
 
 bool UCompositeLayerSceneCapture::IsCustomRenderPass() const
@@ -162,6 +175,18 @@ void UCompositeLayerSceneCapture::SetCustomRenderPass(bool bInIsFastRenderPass)
 	}
 }
 
+bool UCompositeLayerSceneCapture::IsVisibileInSceneCaptureOnly() const
+{
+	return bVisibleInSceneCaptureOnly;
+}
+
+void UCompositeLayerSceneCapture::SetVisibileInSceneCaptureOnly(bool bInVisible)
+{
+	bVisibleInSceneCaptureOnly = bInVisible;
+
+	UpdatePrimitiveVisibilityState();
+}
+
 UCompositeSceneCapture2DComponent* UCompositeLayerSceneCapture::GetSceneCaptureComponent() const
 {
 	ACompositeActor* CompositeActor = GetTypedOuter<ACompositeActor>();
@@ -173,3 +198,48 @@ UCompositeSceneCapture2DComponent* UCompositeLayerSceneCapture::GetSceneCaptureC
 	return nullptr;
 }
 
+void UCompositeLayerSceneCapture::CachePrimitiveVisibilityState(TArrayView<TWeakObjectPtr<UPrimitiveComponent>> InPrimitives)
+{
+	CachedVisibilityInSceneCapture.Reset();
+
+	for (const TWeakObjectPtr<UPrimitiveComponent>& PrimitiveComponent : InPrimitives)
+	{
+		TStrongObjectPtr<UPrimitiveComponent> Primitive = PrimitiveComponent.Pin();
+		if (Primitive.IsValid())
+		{
+			CachedVisibilityInSceneCapture.Add(Primitive.Get(), Primitive->bVisibleInSceneCaptureOnly);
+		}
+	}
+}
+
+void UCompositeLayerSceneCapture::UpdatePrimitiveVisibilityState() const
+{
+	for (const TSoftObjectPtr<AActor>& SoftActor : Actors)
+	{
+		if (const AActor* Actor = SoftActor.Get())
+		{
+			for (UActorComponent* Component : Actor->GetComponents())
+			{
+				if (UPrimitiveComponent* PrimitiveComponent = Cast<UPrimitiveComponent>(Component))
+				{
+					PrimitiveComponent->SetVisibleInSceneCaptureOnly(bVisibleInSceneCaptureOnly);
+				}
+			}
+		}
+	}
+}
+
+void UCompositeLayerSceneCapture::RestorePrimitiveVisibilityState()
+{
+	for (const TPair<TWeakObjectPtr<UPrimitiveComponent>, bool>& Pair : CachedVisibilityInSceneCapture)
+	{
+		TStrongObjectPtr< UPrimitiveComponent> Primitive = Pair.Key.Pin();
+		if (Primitive.IsValid())
+		{
+			Primitive->SetVisibleInSceneCaptureOnly(Pair.Value);
+		}
+	}
+
+	CachedVisibilityInSceneCapture.Reset();
+}
+
diff --git a/Engine/Plugins/Compositing/Composite/Source/Composite/Public/Layers/CompositeLayerSceneCapture.h b/Engine/Plugins/Compositing/Composite/Source/Composite/Public/Layers/CompositeLayerSceneCapture.h
--- a/Engine/Plugins/Compositing/Composite/Source/Composite/Public/Layers/CompositeLayerSceneCapture.h
+++ b/Engine/Plugins/Compositing/Composite/Source/Composite/Public/Layers/CompositeLayerSceneCapture.h
@@ -9,6 +9,7 @@
 #define UE_API COMPOSITE_API
 
 class AActor;
+class UPrimitiveComponent;
 
 /**
 * Layer for scene capture (or custom render pass) render integration.
@@ -53,6 +54,12 @@ public:
 	UFUNCTION(BlueprintSetter)
 	UE_API void SetCustomRenderPass(bool bInIsFastRenderPass);
 
+	UFUNCTION(BlueprintGetter)
+	UE_API bool IsVisibileInSceneCaptureOnly() const;
+
+	UFUNCTION(BlueprintSetter)
+	UE_API void SetVisibileInSceneCaptureOnly(bool bInVisible);
+
 	/** Sub-passes applied during post-processing. */
 	UPROPERTY(EditAnywhere, BlueprintReadWrite, Instanced, Category = "", meta = (DisallowedClasses = "/Script/Composite.CompositePassDistortion", DisplayAfter = "bCustomRenderPass"))
 	TArray<TObjectPtr<UCompositePassBase>> LayerPasses;
@@ -61,6 +68,18 @@ private:
 	/** Convenience function to access the actor-managed scene capture component. */
 	class UCompositeSceneCapture2DComponent* GetSceneCaptureComponent() const;
 
+	/** Get mesh actors primitives. */
+	TArray<TWeakObjectPtr<UPrimitiveComponent>> GetPrimitives() const;
+
+	/** Cache scene capture visibility on registered mesh actors primitives. */
+	void CachePrimitiveVisibilityState(TArrayView<TWeakObjectPtr<UPrimitiveComponent>> InPrimitives);
+
+	/** Update scene capture visibility on registered mesh actors primitives. */
+	void UpdatePrimitiveVisibilityState() const;
+	
+	/** Restore scene capture visibility on registered mesh actors primitives. */
+	void RestorePrimitiveVisibilityState();
+
 	/**
 	 * List of actors to be rendered from the scene capture.
 	 * 
@@ -69,6 +88,14 @@ private:
 	UPROPERTY(EditAnywhere, BlueprintGetter = GetActors, BlueprintSetter = SetActors, Category = "", meta = (AllowPrivateAccess = true))
 	TArray<TSoftObjectPtr<AActor>> Actors;
 
+	/** Cached visibility in scene capture state to restore initial values. */
+	UPROPERTY()
+	TMap<TWeakObjectPtr<UPrimitiveComponent>, bool> CachedVisibilityInSceneCapture;
+
+	/** Visibility setting applied to registered actors. By default, registered meshes will only be visible in scene captures and no longer in the main render. */
+	UPROPERTY(EditAnywhere, BlueprintGetter = IsVisibileInSceneCaptureOnly, BlueprintSetter = SetVisibileInSceneCaptureOnly, Category = "", meta = (AllowPrivateAccess = true))
+	bool bVisibleInSceneCaptureOnly;
+
 	/**
 	 * Convenience toggle to render the scene capture as a fast/minimal Custom Render Pass, inlined in the main render & without support for lighting.
 	 * 
