Composite: Added missing comments from previous changes with scene capture layer.

diff --git a/Engine/Plugins/Compositing/Composite/Source/Composite/Public/Layers/CompositeLayerSceneCapture.h b/Engine/Plugins/Compositing/Composite/Source/Composite/Public/Layers/CompositeLayerSceneCapture.h
--- a/Engine/Plugins/Compositing/Composite/Source/Composite/Public/Layers/CompositeLayerSceneCapture.h
+++ b/Engine/Plugins/Compositing/Composite/Source/Composite/Public/Layers/CompositeLayerSceneCapture.h
@@ -48,15 +48,19 @@ public:
 	UFUNCTION(BlueprintSetter)
 	UE_API void SetActors(TArray<TSoftObjectPtr<AActor>> InActors);
 
+	/** Get whether the scene capture is rendered as a custom render pass. */
 	UFUNCTION(BlueprintGetter)
 	UE_API bool IsCustomRenderPass() const;
 
+	/** Set whether the scene capture is rendered as a custom render pass. */
 	UFUNCTION(BlueprintSetter)
 	UE_API void SetCustomRenderPass(bool bInIsFastRenderPass);
 
+	/** Get whether registered meshes are marked as visible in scene capture only. */
 	UFUNCTION(BlueprintGetter)
 	UE_API bool IsVisibileInSceneCaptureOnly() const;
 
+	/** Set whether registered meshes are marked as visible in scene capture only. */
 	UFUNCTION(BlueprintSetter)
 	UE_API void SetVisibileInSceneCaptureOnly(bool bInVisible);
 
