Composite: Fixed "Visible" typo in blueprint functions.

#rnx

diff --git a/Engine/Plugins/Compositing/Composite/Source/Composite/Private/Layers/CompositeLayerSceneCapture.cpp b/Engine/Plugins/Compositing/Composite/Source/Composite/Private/Layers/CompositeLayerSceneCapture.cpp
--- a/Engine/Plugins/Compositing/Composite/Source/Composite/Private/Layers/CompositeLayerSceneCapture.cpp
+++ b/Engine/Plugins/Compositing/Composite/Source/Composite/Private/Layers/CompositeLayerSceneCapture.cpp
@@ -55,7 +55,7 @@ void UCompositeLayerSceneCapture::PostEditChangeProperty(FPropertyChangedEvent&
 	}
 	else if (PropertyName == GET_MEMBER_NAME_CHECKED(ThisClass, bVisibleInSceneCaptureOnly))
 	{
-		SetVisibileInSceneCaptureOnly(bVisibleInSceneCaptureOnly);
+		SetVisibleInSceneCaptureOnly(bVisibleInSceneCaptureOnly);
 	}
 }
 
@@ -175,12 +175,12 @@ void UCompositeLayerSceneCapture::SetCustomRenderPass(bool bInIsFastRenderPass)
 	}
 }
 
-bool UCompositeLayerSceneCapture::IsVisibileInSceneCaptureOnly() const
+bool UCompositeLayerSceneCapture::IsVisibleInSceneCaptureOnly() const
 {
 	return bVisibleInSceneCaptureOnly;
 }
 
-void UCompositeLayerSceneCapture::SetVisibileInSceneCaptureOnly(bool bInVisible)
+void UCompositeLayerSceneCapture::SetVisibleInSceneCaptureOnly(bool bInVisible)
 {
 	bVisibleInSceneCaptureOnly = bInVisible;
 
diff --git a/Engine/Plugins/Compositing/Composite/Source/Composite/Public/Layers/CompositeLayerSceneCapture.h b/Engine/Plugins/Compositing/Composite/Source/Composite/Public/Layers/CompositeLayerSceneCapture.h
--- a/Engine/Plugins/Compositing/Composite/Source/Composite/Public/Layers/CompositeLayerSceneCapture.h
+++ b/Engine/Plugins/Compositing/Composite/Source/Composite/Public/Layers/CompositeLayerSceneCapture.h
@@ -58,11 +58,11 @@ public:
 
 	/** Get whether registered meshes are marked as visible in scene capture only. */
 	UFUNCTION(BlueprintGetter)
-	UE_API bool IsVisibileInSceneCaptureOnly() const;
+	UE_API bool IsVisibleInSceneCaptureOnly() const;
 
 	/** Set whether registered meshes are marked as visible in scene capture only. */
 	UFUNCTION(BlueprintSetter)
-	UE_API void SetVisibileInSceneCaptureOnly(bool bInVisible);
+	UE_API void SetVisibleInSceneCaptureOnly(bool bInVisible);
 
 	/** Sub-passes applied during post-processing. */
 	UPROPERTY(EditAnywhere, BlueprintReadWrite, Instanced, Category = "", meta = (DisallowedClasses = "/Script/Composite.CompositePassDistortion", DisplayAfter = "bCustomRenderPass"))
@@ -97,7 +97,7 @@ private:
 	TMap<TWeakObjectPtr<UPrimitiveComponent>, bool> CachedVisibilityInSceneCapture;
 
 	/** Visibility setting applied to registered actors. By default, registered meshes will only be visible in scene captures and no longer in the main render. */
-	UPROPERTY(EditAnywhere, BlueprintGetter = IsVisibileInSceneCaptureOnly, BlueprintSetter = SetVisibileInSceneCaptureOnly, Category = "", meta = (AllowPrivateAccess = true))
+	UPROPERTY(EditAnywhere, BlueprintGetter = IsVisibleInSceneCaptureOnly, BlueprintSetter = SetVisibleInSceneCaptureOnly, Category = "", meta = (AllowPrivateAccess = true))
 	bool bVisibleInSceneCaptureOnly;
 
 	/**
