Composite: Replicate primitive visibility caching to shadow & reflection catcher layer.

* Fixed MU & undo support with required Modify() calls.
* Fixed shadow & reflection layer not previously iterating over all components.

diff --git a/Engine/Plugins/Compositing/Composite/Source/Composite/Private/Layers/CompositeLayerSceneCapture.cpp b/Engine/Plugins/Compositing/Composite/Source/Composite/Private/Layers/CompositeLayerSceneCapture.cpp
--- a/Engine/Plugins/Compositing/Composite/Source/Composite/Private/Layers/CompositeLayerSceneCapture.cpp
+++ b/Engine/Plugins/Compositing/Composite/Source/Composite/Private/Layers/CompositeLayerSceneCapture.cpp
@@ -59,6 +59,11 @@ void UCompositeLayerSceneCapture::PostEditChangeProperty(FPropertyChangedEvent&
 	}
 }
 
+void UCompositeLayerSceneCapture::PreEditUndo()
+{
+	Super::PreEditUndo();
+}
+
 void UCompositeLayerSceneCapture::PostEditUndo()
 {
 	Super::PostEditUndo();
@@ -140,7 +145,6 @@ void UCompositeLayerSceneCapture::SetActors(TArray<TSoftObjectPtr<AActor>> InAct
 		SceneCaptureComponent->ShowOnlyComponents = ShowOnlyComponents;
 	}
 
-	// Note: Is it important to restore the state first since this may be called during undo/redo
 	RestorePrimitiveVisibilityState();
 
 	Actors = MoveTemp(InActors);
@@ -212,16 +216,17 @@ void UCompositeLayerSceneCapture::CachePrimitiveVisibilityState(TArrayView<TWeak
 	}
 }
 
-void UCompositeLayerSceneCapture::UpdatePrimitiveVisibilityState() const
+void UCompositeLayerSceneCapture::UpdatePrimitiveVisibilityState()
 {
-	for (const TSoftObjectPtr<AActor>& SoftActor : Actors)
+	for (TSoftObjectPtr<AActor>& SoftActor : Actors)
 	{
-		if (const AActor* Actor = SoftActor.Get())
+		if (AActor* Actor = SoftActor.Get())
 		{
 			for (UActorComponent* Component : Actor->GetComponents())
 			{
 				if (UPrimitiveComponent* PrimitiveComponent = Cast<UPrimitiveComponent>(Component))
 				{
+					PrimitiveComponent->Modify();
 					PrimitiveComponent->SetVisibleInSceneCaptureOnly(bVisibleInSceneCaptureOnly);
 				}
 			}
@@ -236,6 +241,7 @@ void UCompositeLayerSceneCapture::RestorePrimitiveVisibilityState()
 		TStrongObjectPtr< UPrimitiveComponent> Primitive = Pair.Key.Pin();
 		if (Primitive.IsValid())
 		{
+			Primitive->Modify();
 			Primitive->SetVisibleInSceneCaptureOnly(Pair.Value);
 		}
 	}
diff --git a/Engine/Plugins/Compositing/Composite/Source/Composite/Private/Layers/CompositeLayerShadowReflection.cpp b/Engine/Plugins/Compositing/Composite/Source/Composite/Private/Layers/CompositeLayerShadowReflection.cpp
--- a/Engine/Plugins/Compositing/Composite/Source/Composite/Private/Layers/CompositeLayerShadowReflection.cpp
+++ b/Engine/Plugins/Compositing/Composite/Source/Composite/Private/Layers/CompositeLayerShadowReflection.cpp
@@ -19,30 +19,14 @@ UCompositeLayerShadowReflection::~UCompositeLayerShadowReflection() = default;
 
 void UCompositeLayerShadowReflection::OnRemoved()
 {
-	if (AutoConfigureActors == ECompositeHiddenInSceneCaptureConfiguration::Hidden)
-	{
-		for (TSoftObjectPtr<AActor>& Actor : Actors)
-		{
-			if (Actor.IsValid())
-			{
-				if (UPrimitiveComponent* PrimitiveComponent = Actor->GetComponentByClass<UPrimitiveComponent>())
-				{
-					constexpr bool bHideActorsInSceneCapture = false;
-
-					PrimitiveComponent->SetHiddenInSceneCapture(bHideActorsInSceneCapture);
-					PrimitiveComponent->SetAffectIndirectLightingWhileHidden(bHideActorsInSceneCapture);
-					PrimitiveComponent->SetCastHiddenShadow(bHideActorsInSceneCapture);
-				}
-			}
-		}
-	}
-
 	ACompositeActor* CompositeActor = GetTypedOuter<ACompositeActor>();
 	// Note: We don't use IsValid(..) since destruction should proceed even if the actor is pending kill.
 	if (CompositeActor != nullptr)
 	{
 		CompositeActor->DestroySceneCaptures(this);
 	}
+
+	RestorePrimitiveVisibilityState();
 }
 
 void UCompositeLayerShadowReflection::BeginDestroy()
@@ -65,17 +49,20 @@ void UCompositeLayerShadowReflection::PostEditChangeProperty(FPropertyChangedEve
 	}
 	else if (PropertyName == GET_MEMBER_NAME_CHECKED(ThisClass, AutoConfigureActors))
 	{
-		ConfigurePrimitivesForHiddenShadowReflectionCatching();
+		UpdatePrimitiveVisibilityState();
 	}
 }
 
+void UCompositeLayerShadowReflection::PreEditUndo()
+{
+	Super::PreEditUndo();
+}
+
 void UCompositeLayerShadowReflection::PostEditUndo()
 {
 	Super::PostEditUndo();
 
 	SetActors(MoveTemp(Actors));
-
-	ConfigurePrimitivesForHiddenShadowReflectionCatching();
 }
 #endif
 
@@ -137,9 +124,31 @@ const TArray<TSoftObjectPtr<AActor>> UCompositeLayerShadowReflection::GetActors(
 
 void UCompositeLayerShadowReflection::SetActors(TArray<TSoftObjectPtr<AActor>> InActors)
 {
+	RestorePrimitiveVisibilityState();
+
 	Actors = MoveTemp(InActors);
 
-	ConfigurePrimitivesForHiddenShadowReflectionCatching();
+	// Update hidden components
+	TArray<TWeakObjectPtr<UPrimitiveComponent> > HiddenComponents;
+	HiddenComponents.Reserve(Actors.Num());
+
+	for (const TSoftObjectPtr<AActor>& SoftActor : Actors)
+	{
+		if (const AActor* Actor = SoftActor.Get())
+		{
+			for (UActorComponent* Component : Actor->GetComponents())
+			{
+				if (UPrimitiveComponent* PrimitiveComponent = Cast<UPrimitiveComponent>(Component))
+				{
+					HiddenComponents.Add(PrimitiveComponent);
+				}
+			}
+		}
+	}
+
+	CachePrimitiveVisibilityState(HiddenComponents);
+
+	UpdatePrimitiveVisibilityState();
 
 	ACompositeActor* CompositeActor = GetTypedOuter<ACompositeActor>();
 	if (!IsValid(CompositeActor))
@@ -160,29 +169,27 @@ void UCompositeLayerShadowReflection::SetActors(TArray<TSoftObjectPtr<AActor>> I
 
 	// First, we empty hidden actors since our previous logic used them
 	CachedSceneCaptures[1]->HiddenActors.Empty();
+	CachedSceneCaptures[1]->HiddenComponents = HiddenComponents;
+}
 
-	// Update hidden components
-	TArray<TWeakObjectPtr<UPrimitiveComponent> > HiddenComponents;
-	HiddenComponents.Reserve(Actors.Num());
+void UCompositeLayerShadowReflection::CachePrimitiveVisibilityState(TArrayView<TWeakObjectPtr<UPrimitiveComponent>> InPrimitives)
+{
+	CachedVisibilityInSceneCapture.Reset();
 
-	for (const TSoftObjectPtr<AActor>& SoftActor : Actors)
+	for (const TWeakObjectPtr<UPrimitiveComponent>& PrimitiveComponent : InPrimitives)
 	{
-		if (const AActor* Actor = SoftActor.Get())
+		TStrongObjectPtr<UPrimitiveComponent> Primitive = PrimitiveComponent.Pin();
+		if (Primitive.IsValid())
 		{
-			for (UActorComponent* Component : Actor->GetComponents())
-			{
-				if (UPrimitiveComponent* PrimitiveComponent = Cast<UPrimitiveComponent>(Component))
-				{
-					HiddenComponents.Add(PrimitiveComponent);
-				}
-			}
+			FCompositeShadowReflectionPrimitiveState& State = CachedVisibilityInSceneCapture.Add(Primitive.Get());
+			State.bHiddenInSceneCapture = Primitive->bHiddenInSceneCapture;
+			State.bAffectIndirectLightingWhileHidden = Primitive->bAffectIndirectLightingWhileHidden;
+			State.bCastHiddenShadow = Primitive->bCastHiddenShadow;
 		}
 	}
-
-	CachedSceneCaptures[1]->HiddenComponents = HiddenComponents;
 }
 
-void UCompositeLayerShadowReflection::ConfigurePrimitivesForHiddenShadowReflectionCatching()
+void UCompositeLayerShadowReflection::UpdatePrimitiveVisibilityState()
 {
 	bool bHideActorsInSceneCapture = false;
 
@@ -202,16 +209,39 @@ void UCompositeLayerShadowReflection::ConfigurePrimitivesForHiddenShadowReflecti
 	{
 		if (Actor.IsValid())
 		{
-			if (UPrimitiveComponent* PrimitiveComponent = Actor->GetComponentByClass<UPrimitiveComponent>())
+			for (UActorComponent* Component : Actor->GetComponents())
 			{
-				PrimitiveComponent->SetHiddenInSceneCapture(bHideActorsInSceneCapture);
-				PrimitiveComponent->SetAffectIndirectLightingWhileHidden(bHideActorsInSceneCapture);
-				PrimitiveComponent->SetCastHiddenShadow(bHideActorsInSceneCapture);
+				if (UPrimitiveComponent* PrimitiveComponent = Cast<UPrimitiveComponent>(Component))
+				{
+					PrimitiveComponent->Modify();
+					PrimitiveComponent->SetHiddenInSceneCapture(bHideActorsInSceneCapture);
+					PrimitiveComponent->SetAffectIndirectLightingWhileHidden(bHideActorsInSceneCapture);
+					PrimitiveComponent->SetCastHiddenShadow(bHideActorsInSceneCapture);
+				}
 			}
 		}
 	}
 }
 
+void UCompositeLayerShadowReflection::RestorePrimitiveVisibilityState()
+{
+	for (const TPair<TWeakObjectPtr<UPrimitiveComponent>, FCompositeShadowReflectionPrimitiveState>& Pair : CachedVisibilityInSceneCapture)
+	{
+		TStrongObjectPtr< UPrimitiveComponent> Primitive = Pair.Key.Pin();
+		if (Primitive.IsValid())
+		{
+			const FCompositeShadowReflectionPrimitiveState& State = Pair.Value;
+
+			Primitive->Modify();
+			Primitive->SetHiddenInSceneCapture(State.bHiddenInSceneCapture);
+			Primitive->SetAffectIndirectLightingWhileHidden(State.bAffectIndirectLightingWhileHidden);
+			Primitive->SetCastHiddenShadow(State.bCastHiddenShadow);
+		}
+	}
+
+	CachedVisibilityInSceneCapture.Reset();
+}
+
 TStaticArray<TWeakObjectPtr<UCompositeShadowReflectionCatcherComponent>, 2> UCompositeLayerShadowReflection::FindOrCreateSceneCapturePair(ACompositeActor& InOuter) const
 {
 	return {
diff --git a/Engine/Plugins/Compositing/Composite/Source/Composite/Public/Layers/CompositeLayerSceneCapture.h b/Engine/Plugins/Compositing/Composite/Source/Composite/Public/Layers/CompositeLayerSceneCapture.h
--- a/Engine/Plugins/Compositing/Composite/Source/Composite/Public/Layers/CompositeLayerSceneCapture.h
+++ b/Engine/Plugins/Compositing/Composite/Source/Composite/Public/Layers/CompositeLayerSceneCapture.h
@@ -33,6 +33,7 @@ public:
 
 #if WITH_EDITOR
 	UE_API virtual void PostEditChangeProperty(FPropertyChangedEvent& PropertyChangedEvent) override;
+	UE_API virtual void PreEditUndo() override;
 	UE_API virtual void PostEditUndo() override;
 #endif
 	//~ End UObject Interface
@@ -72,14 +73,11 @@ private:
 	/** Convenience function to access the actor-managed scene capture component. */
 	class UCompositeSceneCapture2DComponent* GetSceneCaptureComponent() const;
 
-	/** Get mesh actors primitives. */
-	TArray<TWeakObjectPtr<UPrimitiveComponent>> GetPrimitives() const;
-
 	/** Cache scene capture visibility on registered mesh actors primitives. */
 	void CachePrimitiveVisibilityState(TArrayView<TWeakObjectPtr<UPrimitiveComponent>> InPrimitives);
 
 	/** Update scene capture visibility on registered mesh actors primitives. */
-	void UpdatePrimitiveVisibilityState() const;
+	void UpdatePrimitiveVisibilityState();
 	
 	/** Restore scene capture visibility on registered mesh actors primitives. */
 	void RestorePrimitiveVisibilityState();
diff --git a/Engine/Plugins/Compositing/Composite/Source/Composite/Public/Layers/CompositeLayerShadowReflection.h b/Engine/Plugins/Compositing/Composite/Source/Composite/Public/Layers/CompositeLayerShadowReflection.h
--- a/Engine/Plugins/Compositing/Composite/Source/Composite/Public/Layers/CompositeLayerShadowReflection.h
+++ b/Engine/Plugins/Compositing/Composite/Source/Composite/Public/Layers/CompositeLayerShadowReflection.h
@@ -16,7 +16,7 @@ class ACompositeActor;
 class UCompositeShadowReflectionCatcherComponent;
 class USceneCaptureComponent2D;
 
-/* Automatic primitive configuration mode for registered primitives. */
+/** Automatic primitive configuration mode for registered primitives. */
 UENUM(BlueprintType)
 enum class ECompositeHiddenInSceneCaptureConfiguration : uint8
 {
@@ -25,6 +25,26 @@ enum class ECompositeHiddenInSceneCaptureConfiguration : uint8
 	Hidden UMETA(ToolTip = "The following properties are set to true: bHiddenInSceneCapture, bAffectIndirectLightingWhileHidden & bCastHiddenShadow."),
 };
 
+/** Primitive visibility settings cached state. */
+USTRUCT()
+struct FCompositeShadowReflectionPrimitiveState
+{
+public:
+	GENERATED_USTRUCT_BODY()
+
+	/** Hidden in scene capture. */
+	UPROPERTY()
+	bool bHiddenInSceneCapture = false;
+	
+	/** Affect indirect lighting while hidden. */
+	UPROPERTY()
+	bool bAffectIndirectLightingWhileHidden = false;
+	
+	/** Cast hidden shadow. */
+	UPROPERTY()
+	bool bCastHiddenShadow = false;
+};
+
 /**
 * Layer for shadow & reflection catching with two additional scene capture renders.
 * Primitives are rendered in one scene capture but not the other. The division of both creates a multiplicative matte that can be applied onto the plate.
@@ -48,6 +68,7 @@ public:
 
 #if WITH_EDITOR
 	UE_API virtual void PostEditChangeProperty(FPropertyChangedEvent& PropertyChangedEvent) override;
+	UE_API virtual void PreEditUndo() override;
 	UE_API virtual void PostEditUndo() override;
 #endif
 	//~ End UObject Interface
@@ -74,8 +95,14 @@ public:
 	ECompositeHiddenInSceneCaptureConfiguration AutoConfigureActors;
 
 private:
-	/** Convenience function to configure actor primitives. */
-	void ConfigurePrimitivesForHiddenShadowReflectionCatching();
+	/** Cache scene capture visibility on registered mesh actors primitives. */
+	void CachePrimitiveVisibilityState(TArrayView<TWeakObjectPtr<UPrimitiveComponent>> InPrimitives);
+
+	/** Update scene capture visibility on registered mesh actors primitives. */
+	void UpdatePrimitiveVisibilityState();
+
+	/** Restore scene capture visibility on registered mesh actors primitives. */
+	void RestorePrimitiveVisibilityState();
 
 	/** Find or create scene captures with CG and without CG. */
 	TStaticArray<TWeakObjectPtr<UCompositeShadowReflectionCatcherComponent>, 2> FindOrCreateSceneCapturePair(ACompositeActor& InOuter) const;
@@ -85,6 +112,10 @@ private:
 	UPROPERTY(EditAnywhere, BlueprintGetter = GetActors, BlueprintSetter = SetActors, Category = "")
 	TArray<TSoftObjectPtr<AActor>> Actors;
 
+	/** Cached visibility in scene capture state to restore initial values. */
+	UPROPERTY()
+	TMap<TWeakObjectPtr<UPrimitiveComponent>, FCompositeShadowReflectionPrimitiveState> CachedVisibilityInSceneCapture;
+
 	/** Cached pointer to the scene captures with CG and without CG. */
 	mutable TStaticArray<TWeakObjectPtr<UCompositeShadowReflectionCatcherComponent>, 2> CachedSceneCaptures;
 
