Composure: Added DisplayName field to Plate passes, and allowed passes to be renamed in the Plate Passes tree

diff --git a/Engine/Plugins/Compositing/Composite/Source/Composite/Private/Passes/CompositePassBase.cpp b/Engine/Plugins/Compositing/Composite/Source/Composite/Private/Passes/CompositePassBase.cpp
--- a/Engine/Plugins/Compositing/Composite/Source/Composite/Private/Passes/CompositePassBase.cpp
+++ b/Engine/Plugins/Compositing/Composite/Source/Composite/Private/Passes/CompositePassBase.cpp
@@ -6,6 +6,9 @@
 
 UCompositePassBase::UCompositePassBase(const FObjectInitializer& ObjectInitializer)
 	: Super(ObjectInitializer)
+#if WITH_EDITOR
+	, DisplayName(GetName())
+#endif
 {
 	// Ensure instances are always transactional
 	if (!HasAnyFlags(RF_ClassDefaultObject))
diff --git a/Engine/Plugins/Compositing/Composite/Source/Composite/Public/Passes/CompositePassBase.h b/Engine/Plugins/Compositing/Composite/Source/Composite/Public/Passes/CompositePassBase.h
--- a/Engine/Plugins/Compositing/Composite/Source/Composite/Public/Passes/CompositePassBase.h
+++ b/Engine/Plugins/Compositing/Composite/Source/Composite/Public/Passes/CompositePassBase.h
@@ -46,10 +46,24 @@ public:
 	/** Returns true if the pass (may) need post-processing scene textures. */
 	virtual bool NeedsSceneTextures() const { return false; }
 
+#if WITH_EDITOR
+	/** Gets the display name of the pass */
+	FString GetDisplayName() const { return DisplayName; }
+
+	/** Sets the display name of the pass */
+	void SetDisplayName(const FString& InDisplayName) { DisplayName = InDisplayName; }
+#endif
+	
 protected:
 	/** Whether or not the pass is active. */
 	UPROPERTY(EditAnywhere, BlueprintGetter = IsEnabled, BlueprintSetter = SetEnabled, Category = "", meta = (DisplayPriority = "1"))
 	bool bIsEnabled = true;
+
+#if WITH_EDITORONLY_DATA
+	/** The display name of the pass */
+	UPROPERTY(EditAnywhere, Category = "", meta = (DisplayName="Name", DisplayPriority = "0"))
+	FString DisplayName;
+#endif
 };
 
 #undef UE_API
diff --git a/Engine/Plugins/Compositing/Composite/Source/CompositeEditor/Private/Customizations/CompositeLayerPlateCustomization.cpp b/Engine/Plugins/Compositing/Composite/Source/CompositeEditor/Private/Customizations/CompositeLayerPlateCustomization.cpp
--- a/Engine/Plugins/Compositing/Composite/Source/CompositeEditor/Private/Customizations/CompositeLayerPlateCustomization.cpp
+++ b/Engine/Plugins/Compositing/Composite/Source/CompositeEditor/Private/Customizations/CompositeLayerPlateCustomization.cpp
@@ -141,6 +141,7 @@ void FCompositeLayerPlateCustomization::CustomizeTexturePropertyRow(const TShare
 	[
 		SNew(SObjectPropertyEntryBox)
 		.PropertyHandle(InPropertyHandle)
+		.AllowedClass(UTexture::StaticClass())
 		.AllowCreate(true)
 		.ThumbnailPool(CachedDetailBuilder.IsValid() ? CachedDetailBuilder.Pin()->GetThumbnailPool() : nullptr)
 		.OnShouldFilterAsset_Lambda([](FAssetData const& AssetData)
diff --git a/Engine/Plugins/Compositing/Composite/Source/CompositeEditor/Private/UI/SCompositePassTree.cpp b/Engine/Plugins/Compositing/Composite/Source/CompositeEditor/Private/UI/SCompositePassTree.cpp
--- a/Engine/Plugins/Compositing/Composite/Source/CompositeEditor/Private/UI/SCompositePassTree.cpp
+++ b/Engine/Plugins/Compositing/Composite/Source/CompositeEditor/Private/UI/SCompositePassTree.cpp
@@ -31,6 +31,7 @@ namespace CompositePassTree
 {
 	static const FName PassTreeColumn_Enabled = "Enabled";
 	static const FName PassTreeColumn_Name = "Name";
+	static const FName PassTreeColumn_Type = "Type";
 }
 
 /**
@@ -551,6 +552,15 @@ public:
 		
 		SMultiColumnTableRow::Construct(Args, InOwnerTable);
 	}
+
+	/** Sets the name text block to be in edit mode */
+	void RequestRename()
+	{
+		if (NameTextBlock.IsValid())
+		{
+			NameTextBlock->EnterEditingMode();
+		}
+	}
 	
 private:
 	virtual TSharedRef<SWidget> GenerateWidgetForColumn(const FName& InColumnName) override
@@ -595,7 +605,20 @@ private:
 				+SHorizontalBox::Slot()
 				.VAlign(VAlign_Center)
 				[
-					SNew(STextBlock).Text(this, &SCompositePassTreeItemRow::GetItemName)
+					SAssignNew(NameTextBlock, SInlineEditableTextBlock)
+					.Text(this, &SCompositePassTreeItemRow::GetItemName)
+					.OnTextCommitted(this, &SCompositePassTreeItemRow::SetItemName)
+					.IsReadOnly(this, &SCompositePassTreeItemRow::IsItemNameReadOnly)
+					.IsSelected(FIsSelected::CreateSP(this, &SCompositePassTreeItemRow::IsSelectedExclusively))
+				];
+		}
+		if (InColumnName == CompositePassTree::PassTreeColumn_Type)
+		{
+			return SNew(SBorder)
+				.BorderImage(FAppStyle::GetBrush("NoBorder"))
+				.Padding(4.5f, 0.0f)
+				[
+					SNew(STextBlock).Text(this, &SCompositePassTreeItemRow::GetItemType)
 				];
 		}
 		
@@ -733,11 +756,11 @@ private:
 		{
 			if (UCompositePassBase* Pass = TreeItem->GetPass())
 			{
-				return Pass->GetClass()->GetDisplayNameText();
+				return FText::FromString(Pass->GetDisplayName());
 			}
 			
 			// In the case where there is a valid pass index, but the pass itself is null
-			return LOCTEXT("NoneTypeLabel", "None");
+			return LOCTEXT("PassNotSetLabel", "Pass not set");
 		}
 
 		// Tree item is a pass type group item, so return label for pass type
@@ -751,6 +774,67 @@ private:
 		return PassTypeLabels[(int32)TreeItem->PassType];
 	}
 
+	/** Sets this row item's display name */
+	void SetItemName(const FText& InNewText, ETextCommit::Type InCommitType)
+	{
+		if (!TreeItem->Plate.IsValid())
+		{
+			return;
+		}
+		
+		if (TreeItem->HasValidPassIndex())
+		{
+			if (UCompositePassBase* Pass = TreeItem->GetPass())
+			{
+				FScopedTransaction SetPassNameTransaction(LOCTEXT("SetPassNameTransaction", "Set Name"));
+				Pass->Modify();
+
+				Pass->SetDisplayName(InNewText.ToString());
+			}
+		}
+	}
+
+	/** Gets whether the row item's display name can be changed */
+	bool IsItemNameReadOnly() const
+	{
+		if (!TreeItem->Plate.IsValid())
+		{
+			return true;
+		}
+
+		if (TreeItem->HasValidPassIndex())
+		{
+			if (UCompositePassBase* Pass = TreeItem->GetPass())
+			{
+				return false;
+			}
+		}
+
+		return true;
+	}
+	
+	/** Gets the row item's display name */
+	FText GetItemType() const
+	{
+		if (!TreeItem->Plate.IsValid())
+		{
+			return FText::GetEmpty();
+		}
+		
+		if (TreeItem->HasValidPassIndex())
+		{
+			if (UCompositePassBase* Pass = TreeItem->GetPass())
+			{
+				return Pass->GetClass()->GetDisplayNameText();
+			}
+			
+			// In the case where there is a valid pass index, but the pass itself is null
+            return LOCTEXT("NoneTypeLabel", "None");
+		}
+		
+		return FText::GetEmpty();
+	}
+
 	/** Raised when a drag start has been detected on this row item */
 	FReply HandleDragDetected(const FGeometry& Geometry, const FPointerEvent& MouseEvent)
 	{
@@ -838,6 +922,9 @@ private:
 	/** The tree item this widget is outputting */
 	FTreeItemPtr TreeItem;
 
+	/** The editable text block displaying the item's name */
+	TSharedPtr<SInlineEditableTextBlock> NameTextBlock;
+	
 	/** Callback that is raised when a layer is moved via a drag and drop operation */
 	FOnPassMoved OnPassMoved;
 };
@@ -911,7 +998,12 @@ void SCompositePassTree::Construct(const FArguments& InArgs, const TWeakObjectPt
 				]
 
 				+SHeaderRow::Column(CompositePassTree::PassTreeColumn_Name)
+				.FillWidth(0.55)
 				.DefaultLabel(LOCTEXT("PassNameColumnLabel", "Pass Name"))
+
+				+SHeaderRow::Column(CompositePassTree::PassTreeColumn_Type)
+				.FillWidth(0.45)
+				.DefaultLabel(LOCTEXT("PassTypeColumnLabel", "Pass Type"))
 			)
 			.OnGenerateRow_Lambda([this](FPassTreeItemPtr InTreeItem, const TSharedRef<STableViewBase>& InOwnerTable)
 			{
@@ -1020,6 +1112,10 @@ void SCompositePassTree::BindCommands()
 		FExecuteAction::CreateSP(this, &SCompositePassTree::DeleteSelectedItems),
 		FCanExecuteAction::CreateSP(this, &SCompositePassTree::CanDeleteSelectedItems));
 
+	CommandList->MapAction(FGenericCommands::Get().Rename,
+			FUIAction(FExecuteAction::CreateSP(this, &SCompositePassTree::RenameSelectedItem),
+			FCanExecuteAction::CreateSP(this, &SCompositePassTree::CanRenameSelectedItem)));
+	
 	CommandList->MapAction(FCompositeEditorCommands::Get().Enable,
 		FExecuteAction::CreateSP(this, &SCompositePassTree::EnableSelectedItems),
 		FCanExecuteAction::CreateSP(this, &SCompositePassTree::CanEnableSelectedItems));
@@ -1144,6 +1240,7 @@ TSharedPtr<SWidget> SCompositePassTree::CreateTreeContextMenu()
 		MenuBuilder.AddMenuEntry(FGenericCommands::Get().Paste);
 		MenuBuilder.AddMenuEntry(FGenericCommands::Get().Duplicate);
 		MenuBuilder.AddMenuEntry(FGenericCommands::Get().Delete);
+		MenuBuilder.AddMenuEntry(FGenericCommands::Get().Rename);
 
 		MenuBuilder.AddSeparator();
 
@@ -1802,6 +1899,32 @@ bool SCompositePassTree::CanDeleteSelectedItems()
 	return bContainsPassInSelection;
 }
 
+void SCompositePassTree::RenameSelectedItem()
+{
+	TArray<FPassTreeItemPtr> SelectedItems = TreeView.IsValid() ? TreeView->GetSelectedItems() : TArray<FPassTreeItemPtr>();
+	if (SelectedItems.Num() != 1 || !SelectedItems[0]->HasValidPassIndex() || !SelectedItems[0]->GetPass())
+	{
+		return;
+	}
+	
+	TSharedPtr<ITableRow> RowWidget = TreeView->WidgetFromItem(SelectedItems[0]);
+	if (TSharedPtr<SCompositePassTreeItemRow> TreeItemRowWidget = StaticCastSharedPtr<SCompositePassTreeItemRow>(RowWidget))
+	{
+		TreeItemRowWidget->RequestRename();
+	}
+}
+
+bool SCompositePassTree::CanRenameSelectedItem()
+{
+	TArray<FPassTreeItemPtr> SelectedItems = TreeView.IsValid() ? TreeView->GetSelectedItems() : TArray<FPassTreeItemPtr>();
+	if (SelectedItems.Num() == 1)
+	{
+		return SelectedItems[0].IsValid() && SelectedItems[0]->HasValidPassIndex() && SelectedItems[0]->GetPass();
+	}
+	
+	return false;
+}
+
 void SCompositePassTree::EnableSelectedItems()
 {
 	TArray<FPassTreeItemPtr> SelectedItems = TreeView.IsValid() ? TreeView->GetSelectedItems() : TArray<FPassTreeItemPtr>();
diff --git a/Engine/Plugins/Compositing/Composite/Source/CompositeEditor/Private/UI/SCompositePassTree.h b/Engine/Plugins/Compositing/Composite/Source/CompositeEditor/Private/UI/SCompositePassTree.h
--- a/Engine/Plugins/Compositing/Composite/Source/CompositeEditor/Private/UI/SCompositePassTree.h
+++ b/Engine/Plugins/Compositing/Composite/Source/CompositeEditor/Private/UI/SCompositePassTree.h
@@ -136,6 +136,8 @@ private:
 	bool CanDuplicateSelectedItems();
 	void DeleteSelectedItems();
 	bool CanDeleteSelectedItems();
+	void RenameSelectedItem();
+	bool CanRenameSelectedItem();
 	void EnableSelectedItems();
 	bool CanEnableSelectedItems();
 	
