Composure: Changed default naming scheme for layers and passes to adhere to Actor naming conventions where a number is only added to duplicate objects beyond the first object

diff --git a/Engine/Plugins/Compositing/Composite/Source/Composite/Private/Layers/CompositeLayerBase.cpp b/Engine/Plugins/Compositing/Composite/Source/Composite/Private/Layers/CompositeLayerBase.cpp
--- a/Engine/Plugins/Compositing/Composite/Source/Composite/Private/Layers/CompositeLayerBase.cpp
+++ b/Engine/Plugins/Compositing/Composite/Source/Composite/Private/Layers/CompositeLayerBase.cpp
@@ -13,6 +13,16 @@ UCompositeLayerBase::UCompositeLayerBase(const FObjectInitializer& ObjectInitial
 	{
 		SetFlags(RF_Transactional);
 	}
+
+#if WITH_EDITOR
+	Name = GetClass()->GetName();
+	
+	const int32 NameNumber = GetFName().GetNumber();
+	if (NameNumber > 1)
+	{
+		Name.AppendInt(NameNumber);
+	}
+#endif
 }
 
 UCompositeLayerBase::~UCompositeLayerBase() = default;
diff --git a/Engine/Plugins/Compositing/Composite/Source/Composite/Private/Passes/CompositePassBase.cpp b/Engine/Plugins/Compositing/Composite/Source/Composite/Private/Passes/CompositePassBase.cpp
--- a/Engine/Plugins/Compositing/Composite/Source/Composite/Private/Passes/CompositePassBase.cpp
+++ b/Engine/Plugins/Compositing/Composite/Source/Composite/Private/Passes/CompositePassBase.cpp
@@ -6,15 +6,22 @@
 
 UCompositePassBase::UCompositePassBase(const FObjectInitializer& ObjectInitializer)
 	: Super(ObjectInitializer)
-#if WITH_EDITOR
-	, DisplayName(GetName())
-#endif
 {
 	// Ensure instances are always transactional
 	if (!HasAnyFlags(RF_ClassDefaultObject))
 	{
 		SetFlags(RF_Transactional);
 	}
+
+#if WITH_EDITOR
+	DisplayName = GetClass()->GetName();
+	
+	const int32 NameNumber = GetFName().GetNumber();
+	if (NameNumber > 1)
+	{
+		DisplayName.AppendInt(NameNumber);
+	}
+#endif
 }
 
 UCompositePassBase::~UCompositePassBase() = default;
diff --git a/Engine/Plugins/Compositing/Composite/Source/CompositeEditor/Private/UI/SCompositePanelLayerTree.cpp b/Engine/Plugins/Compositing/Composite/Source/CompositeEditor/Private/UI/SCompositePanelLayerTree.cpp
--- a/Engine/Plugins/Compositing/Composite/Source/CompositeEditor/Private/UI/SCompositePanelLayerTree.cpp
+++ b/Engine/Plugins/Compositing/Composite/Source/CompositeEditor/Private/UI/SCompositePanelLayerTree.cpp
@@ -1846,7 +1846,6 @@ void SCompositePanelLayerTree::OnLayerAdded(const UClass* InLayerClass)
 	CompositeActor->Modify();
 
 	UCompositeLayerBase* NewLayer = NewObject<UCompositeLayerBase>(CompositeActor.Get(), InLayerClass, NAME_None, RF_Transactional);
-	NewLayer->Name = NewLayer->GetName();
 
 	FProperty* CompositeLayersProperty = FindFProperty<FProperty>(ACompositeActor::StaticClass(), GET_MEMBER_NAME_CHECKED(ACompositeActor, CompositeLayers));
 	CompositeActor->PreEditChange(CompositeLayersProperty);
