Composite: Fixed incorrect enabled state handling on plate layer.

* Removed delegate (limited to a single binder) in favor of directly calling layers.
* Ensure that a plate layer rendering state cannot be re-enabled when the composite actor is disabled/de-activated.

diff --git a/Engine/Plugins/Compositing/Composite/Source/Composite/Private/CompositeActor.cpp b/Engine/Plugins/Compositing/Composite/Source/Composite/Private/CompositeActor.cpp
--- a/Engine/Plugins/Compositing/Composite/Source/Composite/Private/CompositeActor.cpp
+++ b/Engine/Plugins/Compositing/Composite/Source/Composite/Private/CompositeActor.cpp
@@ -844,17 +844,26 @@ void ACompositeActor::JoinConcertSession()
 
 void ACompositeActor::ApplyAllState()
 {
-	UCompositeCoreSubsystem* Subsystem = UWorld::GetSubsystem<UCompositeCoreSubsystem>(GetWorld());
-
-	OnEnabledStateChange.ExecuteIfBound(true, Subsystem);
+	for (TObjectPtr<UCompositeLayerBase>& Layer : CompositeLayers)
+	{
+		if (IsValid(Layer))
+		{
+			Layer->OnOwnerRenderingStateChange();
+		}
+	}
 }
 
 void ACompositeActor::RemoveAllState()
 {
-	UCompositeCoreSubsystem* Subsystem = UWorld::GetSubsystem<UCompositeCoreSubsystem>(GetWorld());
-
-	OnEnabledStateChange.ExecuteIfBound(false, Subsystem);
+	for (TObjectPtr<UCompositeLayerBase>& Layer : CompositeLayers)
+	{
+		if (IsValid(Layer))
+		{
+			Layer->OnOwnerRenderingStateChange();
+		}
+	}
 	
+	UCompositeCoreSubsystem* Subsystem = UWorld::GetSubsystem<UCompositeCoreSubsystem>(GetWorld());
 	if (IsValid(Subsystem))
 	{
 		Subsystem->ResetRenderWork();
diff --git a/Engine/Plugins/Compositing/Composite/Source/Composite/Private/Layers/CompositeLayerPlate.cpp b/Engine/Plugins/Compositing/Composite/Source/Composite/Private/Layers/CompositeLayerPlate.cpp
--- a/Engine/Plugins/Compositing/Composite/Source/Composite/Private/Layers/CompositeLayerPlate.cpp
+++ b/Engine/Plugins/Compositing/Composite/Source/Composite/Private/Layers/CompositeLayerPlate.cpp
@@ -101,7 +101,10 @@ void UCompositeLayerPlate::SetEnabled(bool bInEnabled)
 {
 	Super::SetEnabled(bInEnabled);
 
-	HandleEnabledStateChange(IsEnabled(), UWorld::GetSubsystem<UCompositeCoreSubsystem>(GetWorld()));
+	const ACompositeActor* CompositeActor = GetTypedOuter<ACompositeActor>();
+	const bool bIsRendering = IsValid(CompositeActor) ? CompositeActor->IsRendering() : false;
+	
+	PropagateStateChange(IsEnabled() && bIsRendering);
 }
 
 void UCompositeLayerPlate::Tick(float DeltaTime)
@@ -111,7 +114,7 @@ void UCompositeLayerPlate::Tick(float DeltaTime)
 		return;
 	}
 
-	ACompositeActor* CompositeActor = GetTypedOuter<ACompositeActor>();
+	const ACompositeActor* CompositeActor = GetTypedOuter<ACompositeActor>();
 	
 	if (!IsValid(CompositeActor) || !CompositeActor->IsRendering())
 	{
@@ -184,21 +187,6 @@ TStatId UCompositeLayerPlate::GetStatId() const
 	RETURN_QUICK_DECLARE_CYCLE_STAT(UCompositeLayerPlate, STATGROUP_Tickables);
 }
 
-void UCompositeLayerPlate::PostInitProperties()
-{
-	Super::PostInitProperties();
-
-	if (!HasAnyFlags(RF_ClassDefaultObject | RF_Transient))
-	{
-		ACompositeActor* CompositeActor = GetTypedOuter<ACompositeActor>();
-		
-		if (ensure(IsValid(CompositeActor)))
-		{
-			CompositeActor->OnEnabledStateChange.BindUObject(this, &UCompositeLayerPlate::HandleEnabledStateChange);
-		}
-	}
-}
-
 void UCompositeLayerPlate::PostLoad()
 {
 	Super::PostLoad();
@@ -227,6 +215,11 @@ void UCompositeLayerPlate::OnRemoved()
 	}
 }
 
+void UCompositeLayerPlate::OnOwnerRenderingStateChange()
+{
+	SetEnabled(IsEnabled());
+}
+
 void UCompositeLayerPlate::BeginDestroy()
 {
 	Super::BeginDestroy();
@@ -740,31 +733,33 @@ void UCompositeLayerPlate::UpdatePrimitiveComponent(UPrimitiveComponent& InPrimi
 	}
 }
 
-void UCompositeLayerPlate::HandleEnabledStateChange(bool bInEnabled, UCompositeCoreSubsystem* InSubsystem) const
+void UCompositeLayerPlate::PropagateStateChange(bool bInEnabled) const
 {
 	if (CompositeMeshes.IsEmpty())
 	{
 		return;
 	}
 
+	UCompositeCoreSubsystem* Subsystem = UWorld::GetSubsystem<UCompositeCoreSubsystem>(GetWorld());
+
 	if (bInEnabled)
 	{
 		constexpr bool bRegisterOnly = true;
 		UpdateCompositeMeshes(bRegisterOnly);
 
-		if (IsValid(InSubsystem))
+		if (IsValid(Subsystem))
 		{
 			// Rely on the scene view extension to mark primivites as holdout, and render them separately.
-			InSubsystem->RegisterPrimitives(GetPrimitives());
+			Subsystem->RegisterPrimitives(GetPrimitives());
 		}
 	}
 	else
 	{
 		UE::Composite::Private::RemoveCompositeAllAssetUserData(CompositeMeshes);
 
-		if (IsValid(InSubsystem))
+		if (IsValid(Subsystem))
 		{
-			InSubsystem->UnregisterPrimitives(GetPrimitives());
+			Subsystem->UnregisterPrimitives(GetPrimitives());
 		}
 	}
 }
diff --git a/Engine/Plugins/Compositing/Composite/Source/Composite/Public/CompositeActor.h b/Engine/Plugins/Compositing/Composite/Source/Composite/Public/CompositeActor.h
--- a/Engine/Plugins/Compositing/Composite/Source/Composite/Public/CompositeActor.h
+++ b/Engine/Plugins/Compositing/Composite/Source/Composite/Public/CompositeActor.h
@@ -19,7 +19,6 @@ class UPrimitiveComponent;
 class UCompositeSceneCapture2DComponent;
 
 DECLARE_DYNAMIC_MULTICAST_DELEGATE(FOnCompositeActorJoinMultiUserSession);
-DECLARE_DELEGATE_TwoParams(FOnCompositeActorEnabledStateChange, bool /* */, UCompositeCoreSubsystem* /* */);
 
 USTRUCT()
 struct FSceneCaptureComponentArray
@@ -121,9 +120,6 @@ public:
 	UFUNCTION(BlueprintSetter)
 	UE_API void SetCompositeLayers(TArray<UCompositeLayerBase*> InLayers);
 
-	/** Bind to this to get a callback when the enabled state changes. */
-	FOnCompositeActorEnabledStateChange OnEnabledStateChange;
-
 private:
 	/** Component responsible for continuously updating a material parameter collection with the composite camera view projection matrix. */
 	UPROPERTY(Instanced)
diff --git a/Engine/Plugins/Compositing/Composite/Source/Composite/Public/Layers/CompositeLayerBase.h b/Engine/Plugins/Compositing/Composite/Source/Composite/Public/Layers/CompositeLayerBase.h
--- a/Engine/Plugins/Compositing/Composite/Source/Composite/Public/Layers/CompositeLayerBase.h
+++ b/Engine/Plugins/Compositing/Composite/Source/Composite/Public/Layers/CompositeLayerBase.h
@@ -52,6 +52,9 @@ public:
 	UE_API virtual void PostEditChangeProperty(FPropertyChangedEvent& PropertyChangedEvent) override;
 #endif
 
+	/** Called when composite actor rendering state changes. */
+	UE_API virtual void OnOwnerRenderingStateChange() {}
+
 	/** Called when the layer is removed or composite actor is destroyed. */
 	UE_API virtual void OnRemoved() {}
 
diff --git a/Engine/Plugins/Compositing/Composite/Source/Composite/Public/Layers/CompositeLayerPlate.h b/Engine/Plugins/Compositing/Composite/Source/Composite/Public/Layers/CompositeLayerPlate.h
--- a/Engine/Plugins/Compositing/Composite/Source/Composite/Public/Layers/CompositeLayerPlate.h
+++ b/Engine/Plugins/Compositing/Composite/Source/Composite/Public/Layers/CompositeLayerPlate.h
@@ -10,7 +10,6 @@
 
 #define UE_API COMPOSITE_API
 
-class UCompositeCoreSubsystem;
 class UPrimitiveComponent;
 class UTexture;
 class FObjectPreSaveContext;
@@ -47,6 +46,8 @@ public:
 
 	UE_API virtual void OnRemoved() override;
 
+	UE_API virtual void OnOwnerRenderingStateChange() override;
+
 	/** Set the enabled state. */
 	UE_API virtual void SetEnabled(bool bInEnabled) override;
 
@@ -70,7 +71,6 @@ public:
 
 public:
 	//~ Begin UObject Interface
-	UE_API virtual void PostInitProperties() override;
 	UE_API virtual void PostLoad() override;
 	UE_API virtual void BeginDestroy() override;
 
@@ -150,8 +150,8 @@ private:
 	/** Update the (media) texture on a primitive component. */
 	void UpdatePrimitiveComponent(UPrimitiveComponent& InPrimitiveComponent) const;
 
-	/** Handle enabled state change to (de)activate the built-in custom render pass for composite meshes. */
-	void HandleEnabledStateChange(bool bInEnabled, UCompositeCoreSubsystem* InSubsystem) const;
+	/** Propagate enabled/rendering state change to (de)activate the built-in custom render pass for composite meshes. */
+	void PropagateStateChange(bool bInEnabled) const;
 
 	/** Convenience function verify if the composite mesh actor is already in use by another layer. */
 	bool IsCompositeMeshActorAlreadyInUse(TSoftObjectPtr<AActor> InCompositeMeshActor) const;
