Composure: Added button in Composite Actor details panel that opens the Composure panel

diff --git a/Engine/Plugins/Compositing/Composite/Source/CompositeEditor/Private/Customizations/CompositeActorCustomization.cpp b/Engine/Plugins/Compositing/Composite/Source/CompositeEditor/Private/Customizations/CompositeActorCustomization.cpp
--- a/Engine/Plugins/Compositing/Composite/Source/CompositeEditor/Private/Customizations/CompositeActorCustomization.cpp
+++ b/Engine/Plugins/Compositing/Composite/Source/CompositeEditor/Private/Customizations/CompositeActorCustomization.cpp
@@ -3,7 +3,18 @@
 #include "CompositeActorCustomization.h"
 
 #include "CompositeActor.h"
+#include "CompositeEditorStyle.h"
+#include "DetailCategoryBuilder.h"
 #include "DetailLayoutBuilder.h"
+#include "DetailWidgetRow.h"
+#include "LevelEditor.h"
+#include "Framework/Docking/TabManager.h"
+#include "Modules/ModuleManager.h"
+#include "UI/SCompositeEditorPanel.h"
+#include "Widgets/Docking/SDockTab.h"
+#include "Widgets/Images/SImage.h"
+#include "Widgets/Input/SButton.h"
+#include "Widgets/Text/STextBlock.h"
 
 #define LOCTEXT_NAMESPACE "FCompositeActorCustomization"
 
@@ -14,7 +25,70 @@ TSharedRef<IDetailCustomization> FCompositeActorCustomization::MakeInstance()
 
 void FCompositeActorCustomization::CustomizeDetails(IDetailLayoutBuilder& DetailLayout)
 {
+	TArray<TWeakObjectPtr<UObject>> Objects;
+	DetailLayout.GetObjectsBeingCustomized(Objects);
+	CompositeActors.Reserve(Objects.Num());
+	for (TWeakObjectPtr<UObject>& Obj : Objects)
+	{
+		TWeakObjectPtr<ACompositeActor> CompositeActor = Cast<ACompositeActor>(Obj.Get());
+		if (CompositeActor.IsValid())
+		{
+			CompositeActors.Add(CompositeActor);
+		}
+	}
+	
 	DetailLayout.HideProperty(GET_MEMBER_NAME_CHECKED(ACompositeActor, bIsActive));
+
+	DetailLayout.EditCategory("Composite")
+	.AddCustomRow(FText::GetEmpty())
+	.WholeRowContent()
+	.HAlign(HAlign_Fill)
+	.VAlign(VAlign_Center)
+	[
+		SNew(SButton)
+		.OnClicked(this, &FCompositeActorCustomization::OpenComposurePanel)
+		.HAlign(HAlign_Center)
+		.VAlign(VAlign_Center)
+		[
+			SNew(SHorizontalBox)
+						
+			+SHorizontalBox::Slot()
+			.AutoWidth()
+			.Padding(FMargin(0, 1))
+			[
+				SNew(SImage)
+				.Image(FCompositeEditorStyle::Get().GetBrush("CompositeEditor.Composure"))
+			]
+			
+			+SHorizontalBox::Slot()
+			.VAlign(VAlign_Center)
+			.AutoWidth()
+			.Padding(FMargin(2, 0, 0, 0))
+			[
+				SNew(STextBlock)
+				.Text(LOCTEXT("OpenComposureButtonLabel", "Open Composure"))
+				.Font(IDetailLayoutBuilder::GetDetailFont())
+			]
+		]
+	];
+}
+
+FReply FCompositeActorCustomization::OpenComposurePanel()
+{
+	const FName ComposurePanelTabName = "CompositeEditorTab";
+	FLevelEditorModule& LevelEditorModule = FModuleManager::LoadModuleChecked<FLevelEditorModule>("LevelEditor");
+	if (TSharedPtr<FTabManager> LevelEditorTabManager = LevelEditorModule.GetLevelEditorTabManager())
+	{
+		if (TSharedPtr<SDockTab> Tab = LevelEditorTabManager->TryInvokeTab(ComposurePanelTabName))
+		{
+			if (TSharedPtr<SCompositeEditorPanel> CompositePanel = StaticCastSharedPtr<SCompositeEditorPanel>(Tab->GetContent().ToSharedPtr()))
+			{
+				CompositePanel->SelectCompositeActors(CompositeActors);
+			}
+		}
+	}
+	
+	return FReply::Handled();
 }
 
 #undef LOCTEXT_NAMESPACE
diff --git a/Engine/Plugins/Compositing/Composite/Source/CompositeEditor/Private/Customizations/CompositeActorCustomization.h b/Engine/Plugins/Compositing/Composite/Source/CompositeEditor/Private/Customizations/CompositeActorCustomization.h
--- a/Engine/Plugins/Compositing/Composite/Source/CompositeEditor/Private/Customizations/CompositeActorCustomization.h
+++ b/Engine/Plugins/Compositing/Composite/Source/CompositeEditor/Private/Customizations/CompositeActorCustomization.h
@@ -3,6 +3,9 @@
 #pragma once
 
 #include "IDetailCustomization.h"
+#include "Input/Reply.h"
+
+class ACompositeActor;
 
 /**
  * Customization for the composite actor.
@@ -16,4 +19,12 @@ public:
 	// IDetailCustomization interface
 	virtual void CustomizeDetails(IDetailLayoutBuilder& DetailLayout) override;
 	// End of IDetailCustomization interface
+
+private:
+	/** Raised when the 'Open Composure' button is clicked, and opens the Composure panel  */
+	FReply OpenComposurePanel();
+
+private:
+	/** List of composite actors being displayed by this customization */
+	TArray<TWeakObjectPtr<ACompositeActor>> CompositeActors;
 };
diff --git a/Engine/Plugins/Compositing/Composite/Source/CompositeEditor/Private/UI/SCompositeEditorPanel.cpp b/Engine/Plugins/Compositing/Composite/Source/CompositeEditor/Private/UI/SCompositeEditorPanel.cpp
--- a/Engine/Plugins/Compositing/Composite/Source/CompositeEditor/Private/UI/SCompositeEditorPanel.cpp
+++ b/Engine/Plugins/Compositing/Composite/Source/CompositeEditor/Private/UI/SCompositeEditorPanel.cpp
@@ -443,6 +443,14 @@ void SCompositeEditorPanel::Construct(const FArguments& InArgs)
 	];
 }
 
+void SCompositeEditorPanel::SelectCompositeActors(const TArray<TWeakObjectPtr<ACompositeActor>>& InCompositeActors)
+{
+	if (LayerTree.IsValid())
+	{
+		LayerTree->SelectCompositeActors(InCompositeActors);
+	}
+}
+
 void SCompositeEditorPanel::OnLayerSelectionChanged(const TArray<UObject*>& SelectedLayers)
 {
 	if (DetailsView.IsValid())
diff --git a/Engine/Plugins/Compositing/Composite/Source/CompositeEditor/Private/UI/SCompositeEditorPanel.h b/Engine/Plugins/Compositing/Composite/Source/CompositeEditor/Private/UI/SCompositeEditorPanel.h
--- a/Engine/Plugins/Compositing/Composite/Source/CompositeEditor/Private/UI/SCompositeEditorPanel.h
+++ b/Engine/Plugins/Compositing/Composite/Source/CompositeEditor/Private/UI/SCompositeEditorPanel.h
@@ -4,6 +4,7 @@
 
 #include "Widgets/SCompoundWidget.h"
 
+class ACompositeActor;
 class IDetailsView;
 class SCompositePanelLayerTree;
 
@@ -25,6 +26,9 @@ public:
 	
 	void Construct(const FArguments& InArgs);
 
+	/** Selects the specified composite actors in the layer tree */
+	void SelectCompositeActors(const TArray<TWeakObjectPtr<ACompositeActor>>& InCompositeActors);
+	
 private:
 	/** Raised when the selected layers in the tree view have changed */
 	void OnLayerSelectionChanged(const TArray<UObject*>& SelectedLayers);
diff --git a/Engine/Plugins/Compositing/Composite/Source/CompositeEditor/Private/UI/SCompositePanelLayerTree.cpp b/Engine/Plugins/Compositing/Composite/Source/CompositeEditor/Private/UI/SCompositePanelLayerTree.cpp
--- a/Engine/Plugins/Compositing/Composite/Source/CompositeEditor/Private/UI/SCompositePanelLayerTree.cpp
+++ b/Engine/Plugins/Compositing/Composite/Source/CompositeEditor/Private/UI/SCompositePanelLayerTree.cpp
@@ -1059,6 +1059,29 @@ void SCompositePanelLayerTree::PostRedo(bool bSuccess)
 	PostUndo(bSuccess);
 }
 
+void SCompositePanelLayerTree::SelectCompositeActors(const TArray<TWeakObjectPtr<ACompositeActor>>& InCompositeActors)
+{
+	if (!TreeView.IsValid())
+	{
+		return;
+	}
+
+	TreeView->ClearSelection();
+	
+	for (const TWeakObjectPtr<ACompositeActor>& Actor : InCompositeActors)
+	{
+		FCompositeActorTreeItemPtr* TreeItem = CompositeActorTreeItems.FindByPredicate([&Actor](const FCompositeActorTreeItemPtr& InTreeItem)
+		{
+			return InTreeItem.IsValid() && InTreeItem->CompositeActor == Actor.Get();
+		});
+
+		if (TreeItem)
+		{
+			TreeView->SetItemSelection(*TreeItem, true);
+		}
+	}
+}
+
 float SCompositePanelLayerTree::GetMinimumHeight() const
 {
 	// Calculate the number of rows the tree view should be displaying
diff --git a/Engine/Plugins/Compositing/Composite/Source/CompositeEditor/Private/UI/SCompositePanelLayerTree.h b/Engine/Plugins/Compositing/Composite/Source/CompositeEditor/Private/UI/SCompositePanelLayerTree.h
--- a/Engine/Plugins/Compositing/Composite/Source/CompositeEditor/Private/UI/SCompositePanelLayerTree.h
+++ b/Engine/Plugins/Compositing/Composite/Source/CompositeEditor/Private/UI/SCompositePanelLayerTree.h
@@ -67,6 +67,9 @@ public:
 	virtual void PostRedo(bool bSuccess) override;
 	// End of FEditorUndoClient
 
+	/** Selects the specified composite actors in the layer tree */
+	void SelectCompositeActors(const TArray<TWeakObjectPtr<ACompositeActor>>& InCompositeActors);
+	
 	/** Computes the minimum height necessary to display the layer tree with all layers visible */
 	float GetMinimumHeight() const;
 	
