Composite: Force update the view projection matrix when changing which actor is active.

* We also disable view projection matrix updates when the component is parented to a composite actor that is not actively rendering.

diff --git a/Engine/Plugins/Compositing/Composite/Source/Composite/Private/Components/CompositeViewProjectionComponent.cpp b/Engine/Plugins/Compositing/Composite/Source/Composite/Private/Components/CompositeViewProjectionComponent.cpp
--- a/Engine/Plugins/Compositing/Composite/Source/Composite/Private/Components/CompositeViewProjectionComponent.cpp
+++ b/Engine/Plugins/Compositing/Composite/Source/Composite/Private/Components/CompositeViewProjectionComponent.cpp
@@ -137,6 +137,13 @@ void UCompositeViewProjectionComponent::UpdateProjection(UWorld* InWorld) const
 		return;
 	}
 
+	const ACompositeActor* CompositeActor = GetTypedOuter<ACompositeActor>();
+	if (IsValid(CompositeActor) && !CompositeActor->IsRendering())
+	{
+		// Since the component can be used without the composite actor, we only stop updates when it is valid but not rendering.
+		return;
+	}
+
 	UMaterialParameterCollection* MPC = MaterialParameterCollection.Get();
 
 	if (IsValid(LastCameraComponent) && IsValid(MPC))
diff --git a/Engine/Plugins/Compositing/Composite/Source/Composite/Private/CompositeActor.cpp b/Engine/Plugins/Compositing/Composite/Source/Composite/Private/CompositeActor.cpp
--- a/Engine/Plugins/Compositing/Composite/Source/Composite/Private/CompositeActor.cpp
+++ b/Engine/Plugins/Compositing/Composite/Source/Composite/Private/CompositeActor.cpp
@@ -621,6 +621,7 @@ void ACompositeActor::SetActive(bool bInActive)
 		RemoveAllState();
 	}
 
+	// Deactivate all other composite actors
 	const UWorld* World = GetWorld();
 
 	if (bIsActive && IsValid(World))
@@ -634,6 +635,12 @@ void ACompositeActor::SetActive(bool bInActive)
 			}
 		}
 	}
+
+	// Update the view projection matrix sent to the material parameter collection when changing which actor is active.
+	if (bIsActive && IsValid(ViewProjectionComponent))
+	{
+		ViewProjectionComponent->ForceUpdate();
+	}
 }
 
 bool ACompositeActor::IsEnabled() const
