Composite: Renamed allowed view modes property.

diff --git a/Engine/Plugins/Compositing/Composite/Source/Composite/Private/CompositeActor.cpp b/Engine/Plugins/Compositing/Composite/Source/Composite/Private/CompositeActor.cpp
--- a/Engine/Plugins/Compositing/Composite/Source/Composite/Private/CompositeActor.cpp
+++ b/Engine/Plugins/Compositing/Composite/Source/Composite/Private/CompositeActor.cpp
@@ -52,14 +52,14 @@ namespace UE
 
 		namespace Private
 		{
-			/** Convert the allowed viewport view mode into a set of view mode indices. */
-			inline TSet<EViewModeIndex> GetViewModeIndices(ECompositeAllowedViewport InMode)
+			/** Convert the allowed view mode setting into a set of view mode indices. */
+			inline TSet<EViewModeIndex> GetViewModeIndices(ECompositeAllowedViewModes InMode)
 			{
 				switch (InMode)
 				{
-				case ECompositeAllowedViewport::Default:
+				case ECompositeAllowedViewModes::Default:
 					return { VMI_Lit, VMI_PathTracing, VMI_Unknown };
-				case ECompositeAllowedViewport::MediaProfileOnly:
+				case ECompositeAllowedViewModes::MediaProfileUnknown:
 					return { VMI_Unknown };
 				default:
 					break;
@@ -210,7 +210,7 @@ namespace UE
 					break;
 				}
 				
-				RenderWork.AllowedViewModes = GetViewModeIndices(InCompositeActor.AllowedViewport);
+				RenderWork.AllowedViewModes = GetViewModeIndices(InCompositeActor.AllowedViewModes);
 
 				// Currently, each layer may use at most 2 scene captures, so we reserve the maximum case.
 				RenderWork.SceneCapturesUpdateQueue.Reserve(InSceneCapturesPerLayer.Num() * 2);
@@ -246,7 +246,7 @@ ACompositeActor::ACompositeActor(const FObjectInitializer& ObjectInitializer)
 	, bOverridesViewUserFlags{true}
 	, ViewUserFlags{1}
 	, MainRenderOutput{ECompositeMainRenderOutputMode::Default}
-	, AllowedViewport{ECompositeAllowedViewport::Default}
+	, AllowedViewModes{ECompositeAllowedViewModes::Default}
 	, bIsModifyingAProperty{false}
 {
 	PrimaryActorTick.bCanEverTick = true;
diff --git a/Engine/Plugins/Compositing/Composite/Source/Composite/Public/CompositeActor.h b/Engine/Plugins/Compositing/Composite/Source/Composite/Public/CompositeActor.h
--- a/Engine/Plugins/Compositing/Composite/Source/Composite/Public/CompositeActor.h
+++ b/Engine/Plugins/Compositing/Composite/Source/Composite/Public/CompositeActor.h
@@ -42,11 +42,11 @@ enum class ECompositeMainRenderOutputMode : uint8
 * Constrain composite render to certain view modes for multi-viewport workflows where unlit or wireframe viewports will not conflict with the composite.
 */
 UENUM(BlueprintType)
-enum class ECompositeAllowedViewport : uint8
+enum class ECompositeAllowedViewModes : uint8
 {
 	Default UMETA(ToolTip = "Compositing is allowed on viewports with Lit, Path Tracing or Unknown view modes."),
-	MediaProfileOnly UMETA(ToolTip = "Compositing is only allowed with the Unknown view mode, the default when media profile does its own rendering."),
-	All UMETA(ToolTip = "Compositing is allowed with all view modes."),
+	MediaProfileUnknown UMETA(DisplayName = "Media Profile (Unknown)", ToolTip = "Compositing is only allowed with the Unknown view mode, the default when media profile does its own rendering."),
+	AllViewModes UMETA(ToolTip = "Compositing is allowed with all view modes."),
 };
 
 /** Actor used to control properties of the composite pipeline. */
@@ -164,9 +164,9 @@ public:
 	UPROPERTY(EditAnywhere, BlueprintReadWrite, Category = "Composite", AdvancedDisplay)
 	ECompositeMainRenderOutputMode MainRenderOutput;
 
-	/** Constrain composite rendering to specific viewport view modes. */
+	/** Constrain composite rendering to specific view modes. */
 	UPROPERTY(EditAnywhere, BlueprintReadWrite, Category = "Composite", AdvancedDisplay)
-	ECompositeAllowedViewport AllowedViewport;
+	ECompositeAllowedViewModes AllowedViewModes;
 
 	/** Delegate called when joining multi-user. Can be used to automatically re-active the actor locally. */
 	UPROPERTY(BlueprintAssignable, BlueprintCallable, BlueprintReadWrite, Category = "Composite")
