Composite: Fixed composite mesh state cleanup/removal upon plate layer deletion (existing issue) or composite actor deletion (caused by recent CL 46554760).

* We now pass the actor's world to layers upon removal for successful call to the world subsystem.
* Replicated details panel instanced property array removal in composure panel by renaming removed object to none/transient.
* Combined Composite Actor ApplyAllState & RemoveAllState methods for simplicity.
* Moved rendering & enabled conditions into PropagateStateChange function to unify conditioning.
* Fixed conditional ticking on plate layer.
* Additional fix to undo solo layer changes.

diff --git a/Engine/Plugins/Compositing/Composite/Source/Composite/Private/CompositeActor.cpp b/Engine/Plugins/Compositing/Composite/Source/Composite/Private/CompositeActor.cpp
--- a/Engine/Plugins/Compositing/Composite/Source/Composite/Private/CompositeActor.cpp
+++ b/Engine/Plugins/Compositing/Composite/Source/Composite/Private/CompositeActor.cpp
@@ -342,7 +342,7 @@ void ACompositeActor::PostRegisterAllComponents()
 
 	if (IsRendering())
 	{
-		ApplyAllState();
+		PropagateStateChange(true);
 	}
 }
 
@@ -356,8 +356,9 @@ void ACompositeActor::UnregisterAllComponents(bool bForReregister)
 	{
 		return;
 	}
-
-	RemoveAllState();
+	
+	// State removal is done without any conditions to prevent stale state
+	PropagateStateChange(false);
 }
 
 void ACompositeActor::PostUnregisterAllComponents()
@@ -399,7 +400,7 @@ void ACompositeActor::Destroyed()
 		// Note: We don't use IsValid(..) since destruction should proceed even if the actor is pending kill.
 		if (Layer != nullptr)
 		{
-			Layer->OnRemoved();
+			Layer->OnRemoved(GetWorld());
 		}
 	}
 }
@@ -505,7 +506,7 @@ void ACompositeActor::PostEditChangeChainProperty(struct FPropertyChangedChainEv
 				{
 					if (!CompositeLayers.Contains(PreEditCompositeLayer))
 					{
-						PreEditCompositeLayer->OnRemoved();
+						PreEditCompositeLayer->OnRemoved(GetWorld());
 					}
 				}
 			}
@@ -551,7 +552,7 @@ void ACompositeActor::PostEditUndo()
 		{
 			if (!CompositeLayers.Contains(PreEditCompositeLayer))
 			{
-				PreEditCompositeLayer->OnRemoved();
+				PreEditCompositeLayer->OnRemoved(GetWorld());
 			}
 		}
 	}
@@ -587,14 +588,7 @@ void ACompositeActor::SetIsTemporarilyHiddenInEditor(bool bIsHidden)
 {
 	Super::SetIsTemporarilyHiddenInEditor(bIsHidden);
 
-	if (IsRendering())
-	{
-		ApplyAllState();
-	}
-	else
-	{
-		RemoveAllState();
-	}
+	PropagateStateChange(IsRendering());
 }
 
 bool ACompositeActor::ShouldTickIfViewportsOnly() const
@@ -612,14 +606,7 @@ void ACompositeActor::SetActive(bool bInActive)
 {
 	bIsActive = bInActive;
 
-	if (IsRendering())
-	{
-		ApplyAllState();
-	}
-	else
-	{
-		RemoveAllState();
-	}
+	PropagateStateChange(IsRendering());
 
 	// Deactivate all other composite actors
 	const UWorld* World = GetWorld();
@@ -631,7 +618,7 @@ void ACompositeActor::SetActive(bool bInActive)
 			if (*It != this)
 			{
 				It->bIsActive = false;
-				It->RemoveAllState();
+				It->PropagateStateChange(false);
 			}
 		}
 	}
@@ -652,14 +639,7 @@ void ACompositeActor::SetEnabled(bool bInEnabled)
 {
 	bIsEnabled = bInEnabled;
 
-	if (IsRendering())
-	{
-		ApplyAllState();
-	}
-	else
-	{
-		RemoveAllState();
-	}
+	PropagateStateChange(IsRendering());
 }
 
 FComponentReference& ACompositeActor::GetCamera()
@@ -695,7 +675,7 @@ void ACompositeActor::SetCompositeLayers(TArray<UCompositeLayerBase*> InLayers)
 		{
 			if (!InLayers.Contains(PreviousLayer))
 			{
-				PreviousLayer->OnRemoved();
+				PreviousLayer->OnRemoved(GetWorld());
 			}
 		}
 	}
@@ -849,33 +829,26 @@ void ACompositeActor::JoinConcertSession()
 	OnJoinMultiUserSession.Broadcast();
 }
 
-void ACompositeActor::ApplyAllState()
+void ACompositeActor::PropagateStateChange(bool bApply)
 {
 	for (TObjectPtr<UCompositeLayerBase>& Layer : CompositeLayers)
 	{
 		if (IsValid(Layer))
 		{
-			Layer->OnOwnerRenderingStateChange();
+			Layer->OnRenderingStateChange(bApply);
 		}
 	}
-}
 
-void ACompositeActor::RemoveAllState()
-{
-	for (TObjectPtr<UCompositeLayerBase>& Layer : CompositeLayers)
+	// Remove state case
+	if (!bApply)
 	{
-		if (IsValid(Layer))
+		UCompositeCoreSubsystem* Subsystem = UWorld::GetSubsystem<UCompositeCoreSubsystem>(GetWorld());
+		if (IsValid(Subsystem))
 		{
-			Layer->OnOwnerRenderingStateChange();
+			Subsystem->ResetRenderWork();
+			Subsystem->ResetBuiltInRenderPassOptions();
 		}
 	}
-	
-	UCompositeCoreSubsystem* Subsystem = UWorld::GetSubsystem<UCompositeCoreSubsystem>(GetWorld());
-	if (IsValid(Subsystem))
-	{
-		Subsystem->ResetRenderWork();
-		Subsystem->ResetBuiltInRenderPassOptions();
-	}
 }
 
 
diff --git a/Engine/Plugins/Compositing/Composite/Source/Composite/Private/Layers/CompositeLayerPlate.cpp b/Engine/Plugins/Compositing/Composite/Source/Composite/Private/Layers/CompositeLayerPlate.cpp
--- a/Engine/Plugins/Compositing/Composite/Source/Composite/Private/Layers/CompositeLayerPlate.cpp
+++ b/Engine/Plugins/Compositing/Composite/Source/Composite/Private/Layers/CompositeLayerPlate.cpp
@@ -100,11 +100,8 @@ UCompositeLayerPlate::~UCompositeLayerPlate() = default;
 void UCompositeLayerPlate::SetEnabled(bool bInEnabled)
 {
 	Super::SetEnabled(bInEnabled);
-
-	const ACompositeActor* CompositeActor = GetTypedOuter<ACompositeActor>();
-	const bool bIsRendering = IsValid(CompositeActor) ? CompositeActor->IsRendering() : false;
 	
-	PropagateStateChange(IsEnabled() && bIsRendering);
+	PropagateStateChange(bInEnabled, GetWorld());
 }
 
 void UCompositeLayerPlate::Tick(float DeltaTime)
@@ -169,12 +166,12 @@ void UCompositeLayerPlate::Tick(float DeltaTime)
 
 ETickableTickType UCompositeLayerPlate::GetTickableTickType() const
 {
-	return HasAnyFlags(RF_ClassDefaultObject) ? ETickableTickType::Never : ETickableTickType::Always;
+	return HasAnyFlags(RF_ClassDefaultObject) ? ETickableTickType::Never : ETickableTickType::Conditional;
 }
 
 bool UCompositeLayerPlate::IsTickable() const
 {
-	return (HasAnyFlags(RF_ClassDefaultObject) == false);
+	return !HasAnyFlags(RF_ClassDefaultObject) && GetTypedOuter<ACompositeActor>();
 }
 
 UWorld* UCompositeLayerPlate::GetTickableGameObjectWorld() const
@@ -198,33 +195,36 @@ void UCompositeLayerPlate::PostLoad()
 	}
 }
 
-void UCompositeLayerPlate::OnRemoved()
+void UCompositeLayerPlate::OnRemoved(const UWorld* World)
 {
 	FCompositeRenderTargetPool::Get().ReleaseTarget(MediaRenderTarget);
 	FCompositeRenderTargetPool::Get().ReleaseTarget(SceneRenderTarget);
 
-	// If we are not the class default object...
-	if (!HasAnyFlags(RF_ClassDefaultObject))
-	{
-		TryCloseMediaProfileSource();
-		
-#if WITH_EDITOR
-		FEditorDelegates::PreSaveWorldWithContext.RemoveAll(this);
-		FEditorDelegates::PostSaveWorldWithContext.RemoveAll(this);
-#endif // WITH_EDITOR
-	}
+	TryCloseMediaProfileSource();
+
+	constexpr bool bApply = false;
+	PropagateStateChange(bApply, World);
 }
 
-void UCompositeLayerPlate::OnOwnerRenderingStateChange()
+void UCompositeLayerPlate::OnRenderingStateChange(bool bApply)
 {
-	SetEnabled(IsEnabled());
+	PropagateStateChange(bApply, GetWorld());
 }
 
 void UCompositeLayerPlate::BeginDestroy()
 {
 	Super::BeginDestroy();
 
-	OnRemoved(); // Redundant remove call for safety
+#if WITH_EDITOR
+	// Only remove delegates on final destruction to preserve them for undo operations.
+	if (!HasAnyFlags(RF_ClassDefaultObject))
+	{
+		FEditorDelegates::PreSaveWorldWithContext.RemoveAll(this);
+		FEditorDelegates::PostSaveWorldWithContext.RemoveAll(this);
+	}
+#endif // WITH_EDITOR
+
+	OnRemoved(GetWorld()); // Redundant remove call for safety
 }
 
 const TArray<TSoftObjectPtr<AActor>> UCompositeLayerPlate::GetCompositeMeshes() const
@@ -244,10 +244,7 @@ void UCompositeLayerPlate::SetCompositeMeshes(TArray<TSoftObjectPtr<AActor>> InC
 
 void UCompositeLayerPlate::SetPlateMode(ECompositePlateMode InPlateMode)
 {
-	if (PlateMode != InPlateMode)
-	{
-		PlateMode = InPlateMode;
-	}
+	PlateMode = InPlateMode;
 }
 
 UTexture* UCompositeLayerPlate::GetCompositeTexture() const
@@ -733,16 +730,18 @@ void UCompositeLayerPlate::UpdatePrimitiveComponent(UPrimitiveComponent& InPrimi
 	}
 }
 
-void UCompositeLayerPlate::PropagateStateChange(bool bInEnabled) const
+void UCompositeLayerPlate::PropagateStateChange(bool bApply, const UWorld* SourceWorld) const
 {
 	if (CompositeMeshes.IsEmpty())
 	{
 		return;
 	}
 
-	UCompositeCoreSubsystem* Subsystem = UWorld::GetSubsystem<UCompositeCoreSubsystem>(GetWorld());
+	UCompositeCoreSubsystem* Subsystem = UWorld::GetSubsystem<UCompositeCoreSubsystem>(SourceWorld ? SourceWorld : GetWorld());
+	const ACompositeActor* CompositeActor = GetTypedOuter<ACompositeActor>();
+	const bool bIsRendering = IsValid(CompositeActor) ? CompositeActor->IsRendering() : false;
 
-	if (bInEnabled)
+	if (bApply && IsEnabled() && bIsRendering)
 	{
 		constexpr bool bRegisterOnly = true;
 		UpdateCompositeMeshes(bRegisterOnly);
diff --git a/Engine/Plugins/Compositing/Composite/Source/Composite/Private/Layers/CompositeLayerSceneCapture.cpp b/Engine/Plugins/Compositing/Composite/Source/Composite/Private/Layers/CompositeLayerSceneCapture.cpp
--- a/Engine/Plugins/Compositing/Composite/Source/Composite/Private/Layers/CompositeLayerSceneCapture.cpp
+++ b/Engine/Plugins/Compositing/Composite/Source/Composite/Private/Layers/CompositeLayerSceneCapture.cpp
@@ -19,7 +19,7 @@ UCompositeLayerSceneCapture::UCompositeLayerSceneCapture(const FObjectInitialize
 
 UCompositeLayerSceneCapture::~UCompositeLayerSceneCapture() = default;
 
-void UCompositeLayerSceneCapture::OnRemoved()
+void UCompositeLayerSceneCapture::OnRemoved(const UWorld* World)
 {
 	ACompositeActor* CompositeActor = GetTypedOuter<ACompositeActor>();
 	// Note: We don't use IsValid(..) since destruction should proceed even if the actor is pending kill.
@@ -35,7 +35,7 @@ void UCompositeLayerSceneCapture::BeginDestroy()
 {
 	Super::BeginDestroy();
 
-	OnRemoved(); // Redundant remove call for safety
+	OnRemoved(GetWorld()); // Redundant remove call for safety
 }
 
 #if WITH_EDITOR
diff --git a/Engine/Plugins/Compositing/Composite/Source/Composite/Private/Layers/CompositeLayerShadowReflection.cpp b/Engine/Plugins/Compositing/Composite/Source/Composite/Private/Layers/CompositeLayerShadowReflection.cpp
--- a/Engine/Plugins/Compositing/Composite/Source/Composite/Private/Layers/CompositeLayerShadowReflection.cpp
+++ b/Engine/Plugins/Compositing/Composite/Source/Composite/Private/Layers/CompositeLayerShadowReflection.cpp
@@ -17,7 +17,7 @@ UCompositeLayerShadowReflection::UCompositeLayerShadowReflection(const FObjectIn
 
 UCompositeLayerShadowReflection::~UCompositeLayerShadowReflection() = default;
 
-void UCompositeLayerShadowReflection::OnRemoved()
+void UCompositeLayerShadowReflection::OnRemoved(const UWorld* World)
 {
 	ACompositeActor* CompositeActor = GetTypedOuter<ACompositeActor>();
 	// Note: We don't use IsValid(..) since destruction should proceed even if the actor is pending kill.
@@ -33,7 +33,7 @@ void UCompositeLayerShadowReflection::BeginDestroy()
 {
 	Super::BeginDestroy();
 
-	OnRemoved(); // Redundant remove call for safety
+	OnRemoved(GetWorld()); // Redundant remove call for safety
 }
 
 #if WITH_EDITOR
diff --git a/Engine/Plugins/Compositing/Composite/Source/Composite/Private/Layers/CompositeLayerSingleLightShadow.cpp b/Engine/Plugins/Compositing/Composite/Source/Composite/Private/Layers/CompositeLayerSingleLightShadow.cpp
--- a/Engine/Plugins/Compositing/Composite/Source/Composite/Private/Layers/CompositeLayerSingleLightShadow.cpp
+++ b/Engine/Plugins/Compositing/Composite/Source/Composite/Private/Layers/CompositeLayerSingleLightShadow.cpp
@@ -164,7 +164,7 @@ UCompositeLayerSingleLightShadow::UCompositeLayerSingleLightShadow(const FObject
 	Operation = ECompositeCoreMergeOp::Multiply;
 }
 
-void UCompositeLayerSingleLightShadow::OnRemoved()
+void UCompositeLayerSingleLightShadow::OnRemoved(const UWorld* World)
 {
 	ACompositeActor* CompositeActor = GetTypedOuter<ACompositeActor>();
 	// Note: We don't use IsValid(..) since destruction should proceed even if the actor is pending kill.
@@ -178,7 +178,7 @@ void UCompositeLayerSingleLightShadow::BeginDestroy()
 {
 	Super::BeginDestroy();
 
-	OnRemoved(); // Redundant remove call for safety
+	OnRemoved(GetWorld()); // Redundant remove call for safety
 }
 
 #if WITH_EDITOR
diff --git a/Engine/Plugins/Compositing/Composite/Source/Composite/Public/CompositeActor.h b/Engine/Plugins/Compositing/Composite/Source/Composite/Public/CompositeActor.h
--- a/Engine/Plugins/Compositing/Composite/Source/Composite/Public/CompositeActor.h
+++ b/Engine/Plugins/Compositing/Composite/Source/Composite/Public/CompositeActor.h
@@ -187,11 +187,12 @@ public:
 	UE_API void JoinConcertSession();
 
 private:
-	/** Apply all actor state to the composite subsystem. */
-	void ApplyAllState();
-
-	/** Remove all actor state from the composite subsystem. */
-	void RemoveAllState();
+	/**
+	* Propagate rendering state changes to layers and the composite core subsystem.
+	* 
+	* @param bApply - Apply or remove external rendering state.
+	*/
+	void PropagateStateChange(bool bApply);
 
 	/** Layer-managed scene capture components. */
 	UPROPERTY()
diff --git a/Engine/Plugins/Compositing/Composite/Source/Composite/Public/Layers/CompositeLayerBase.h b/Engine/Plugins/Compositing/Composite/Source/Composite/Public/Layers/CompositeLayerBase.h
--- a/Engine/Plugins/Compositing/Composite/Source/Composite/Public/Layers/CompositeLayerBase.h
+++ b/Engine/Plugins/Compositing/Composite/Source/Composite/Public/Layers/CompositeLayerBase.h
@@ -14,6 +14,7 @@
 #define UE_API COMPOSITE_API
 
 class ACompositeActor;
+class UWorld;
 
 /** Base class for defining a pass object, and its associated render proxy. */
 UCLASS(MinimalAPI, Abstract, EditInlineNew)
@@ -52,11 +53,19 @@ public:
 	UE_API virtual void PostEditChangeProperty(FPropertyChangedEvent& PropertyChangedEvent) override;
 #endif
 
-	/** Called when composite actor rendering state changes. */
-	UE_API virtual void OnOwnerRenderingStateChange() {}
-
-	/** Called when the layer is removed or composite actor is destroyed. */
-	UE_API virtual void OnRemoved() {}
+	/**
+	* Called when composite actor rendering state changes.
+	* 
+	* @param bApply - Apply or remove external rendering state.
+	*/
+	UE_API virtual void OnRenderingStateChange(bool bApply) {}
+
+	/**
+	* Called when the layer is removed or composite actor is destroyed.
+	*
+	* @param World - World from which the layer is removed. Ignore if null.
+	*/
+	UE_API virtual void OnRemoved(const UWorld* World) {}
 
 	/** Getter function to override, returning pass layer proxies to be passed to the render thread. (Proxy objects should be allocated from the provided allocator.) */
 	UE_API virtual bool GetProxy(FTraversalContext& InContext, FSceneRenderingBulkObjectAllocator& InFrameAllocator, FCompositeCorePassProxy*& OutProxy) const { return false; };
diff --git a/Engine/Plugins/Compositing/Composite/Source/Composite/Public/Layers/CompositeLayerPlate.h b/Engine/Plugins/Compositing/Composite/Source/Composite/Public/Layers/CompositeLayerPlate.h
--- a/Engine/Plugins/Compositing/Composite/Source/Composite/Public/Layers/CompositeLayerPlate.h
+++ b/Engine/Plugins/Compositing/Composite/Source/Composite/Public/Layers/CompositeLayerPlate.h
@@ -44,9 +44,9 @@ public:
 	/** Destructor. */
 	UE_API virtual ~UCompositeLayerPlate();
 
-	UE_API virtual void OnRemoved() override;
+	UE_API virtual void OnRemoved(const UWorld* World) override;
 
-	UE_API virtual void OnOwnerRenderingStateChange() override;
+	UE_API virtual void OnRenderingStateChange(bool bApply) override;
 
 	/** Set the enabled state. */
 	UE_API virtual void SetEnabled(bool bInEnabled) override;
@@ -150,8 +150,13 @@ private:
 	/** Update the (media) texture on a primitive component. */
 	void UpdatePrimitiveComponent(UPrimitiveComponent& InPrimitiveComponent) const;
 
-	/** Propagate enabled/rendering state change to (de)activate the built-in custom render pass for composite meshes. */
-	void PropagateStateChange(bool bInEnabled) const;
+	/**
+	* Apply (or remove) state to activate (or deactivate) the built-in custom render pass for composite meshes.
+	* 
+	* @param bApply Apply or remove external rendering state.
+	* @param SourceWorld Optional source world to find the composite subsystem
+	*/
+	void PropagateStateChange(bool bApply, const UWorld* SourceWorld) const;
 
 	/** Convenience function verify if the composite mesh actor is already in use by another layer. */
 	bool IsCompositeMeshActorAlreadyInUse(TSoftObjectPtr<AActor> InCompositeMeshActor) const;
diff --git a/Engine/Plugins/Compositing/Composite/Source/Composite/Public/Layers/CompositeLayerSceneCapture.h b/Engine/Plugins/Compositing/Composite/Source/Composite/Public/Layers/CompositeLayerSceneCapture.h
--- a/Engine/Plugins/Compositing/Composite/Source/Composite/Public/Layers/CompositeLayerSceneCapture.h
+++ b/Engine/Plugins/Compositing/Composite/Source/Composite/Public/Layers/CompositeLayerSceneCapture.h
@@ -26,7 +26,7 @@ public:
 	/** Destructor. */
 	UE_API virtual ~UCompositeLayerSceneCapture();
 	
-	UE_API virtual void OnRemoved() override;
+	UE_API virtual void OnRemoved(const UWorld* World) override;
 
 	//~ Begin UObject Interface
 	UE_API virtual void BeginDestroy() override;
diff --git a/Engine/Plugins/Compositing/Composite/Source/Composite/Public/Layers/CompositeLayerShadowReflection.h b/Engine/Plugins/Compositing/Composite/Source/Composite/Public/Layers/CompositeLayerShadowReflection.h
--- a/Engine/Plugins/Compositing/Composite/Source/Composite/Public/Layers/CompositeLayerShadowReflection.h
+++ b/Engine/Plugins/Compositing/Composite/Source/Composite/Public/Layers/CompositeLayerShadowReflection.h
@@ -61,7 +61,7 @@ public:
 	/** Destructor. */
 	UE_API virtual ~UCompositeLayerShadowReflection();
 
-	UE_API virtual void OnRemoved() override;
+	UE_API virtual void OnRemoved(const UWorld* World) override;
 
 	//~ Begin UObject Interface
 	UE_API virtual void BeginDestroy() override;
diff --git a/Engine/Plugins/Compositing/Composite/Source/Composite/Public/Layers/CompositeLayerSingleLightShadow.h b/Engine/Plugins/Compositing/Composite/Source/Composite/Public/Layers/CompositeLayerSingleLightShadow.h
--- a/Engine/Plugins/Compositing/Composite/Source/Composite/Public/Layers/CompositeLayerSingleLightShadow.h
+++ b/Engine/Plugins/Compositing/Composite/Source/Composite/Public/Layers/CompositeLayerSingleLightShadow.h
@@ -37,7 +37,7 @@ public:
 	/** Destructor. */
 	UE_API virtual ~UCompositeLayerSingleLightShadow() = default;
 
-	UE_API virtual void OnRemoved() override;
+	UE_API virtual void OnRemoved(const UWorld* World) override;
 
 	//~ Begin UObject Interface
 	UE_API virtual void BeginDestroy() override;
diff --git a/Engine/Plugins/Compositing/Composite/Source/CompositeEditor/Private/UI/SCompositePanelLayerTree.cpp b/Engine/Plugins/Compositing/Composite/Source/CompositeEditor/Private/UI/SCompositePanelLayerTree.cpp
--- a/Engine/Plugins/Compositing/Composite/Source/CompositeEditor/Private/UI/SCompositePanelLayerTree.cpp
+++ b/Engine/Plugins/Compositing/Composite/Source/CompositeEditor/Private/UI/SCompositePanelLayerTree.cpp
@@ -1920,7 +1920,11 @@ FReply SCompositePanelLayerTree::OnLayerSoloToggled(const FCompositeActorTreeIte
 		{
 			if (OtherLayer != CompositeLayer)
 			{
-				OtherLayer->bIsSolo = false;
+				if (IsValid(OtherLayer))
+				{
+					OtherLayer->Modify();
+					OtherLayer->bIsSolo = false;
+				}
 			}
 		}
 	}
@@ -2281,15 +2285,26 @@ void SCompositePanelLayerTree::DeleteSelectedItems()
 			CompositeActor->Modify();
 			
 			// Delete the layers in reverse order so that indices don't get messed up as we delete
-			TArray<int32>& Layers = CompositeActorLayerList.Value;
-			Layers.Sort();
+			TArray<int32>& LayerIndices = CompositeActorLayerList.Value;
+			LayerIndices.Sort();
 
 			FProperty* CompositeLayersProperty = FindFProperty<FProperty>(ACompositeActor::StaticClass(), GET_MEMBER_NAME_CHECKED(ACompositeActor, CompositeLayers));
 			CompositeActor->PreEditChange(CompositeLayersProperty);
 			
-			for (int32 Index = Layers.Num() - 1; Index >= 0; --Index)
+			for (int32 Index = LayerIndices.Num() - 1; Index >= 0; --Index)
 			{
-				CompositeActor->CompositeLayers.RemoveAt(Layers[Index]);
+				const int32 LayerIndex = LayerIndices[Index];
+
+				// Here we replicate details panels instanced property array removal code.
+				// The removed instanced layer objects must be moved to the transient package.
+				TObjectPtr<UCompositeLayerBase>& LayerToRemove = CompositeActor->CompositeLayers[LayerIndex];
+				if (IsValid(LayerToRemove))
+				{
+					LayerToRemove->Modify();
+					LayerToRemove->Rename(nullptr, GetTransientPackage(), REN_DontCreateRedirectors);
+				}
+
+				CompositeActor->CompositeLayers.RemoveAt(LayerIndex);
 			}
 
 			FEditPropertyChain EditChain;
diff --git a/Engine/Plugins/Compositing/Composite/Source/CompositeEditor/Private/UI/SCompositePassTree.cpp b/Engine/Plugins/Compositing/Composite/Source/CompositeEditor/Private/UI/SCompositePassTree.cpp
--- a/Engine/Plugins/Compositing/Composite/Source/CompositeEditor/Private/UI/SCompositePassTree.cpp
+++ b/Engine/Plugins/Compositing/Composite/Source/CompositeEditor/Private/UI/SCompositePassTree.cpp
@@ -1872,7 +1872,18 @@ void SCompositePassTree::DeleteSelectedItems()
 			
 			for (int32 Index = PassIndicesToDelete.Num() - 1; Index >= 0; --Index)
 			{
-				PassList.RemoveAt(PassIndicesToDelete[Index], EAllowShrinking::No);
+				const int32 PassIndex = PassIndicesToDelete[Index];
+
+				// Here we replicate details panels instanced property array removal code.
+				// The removed instanced pass objects must be moved to the transient package.
+				TObjectPtr<UCompositePassBase>& PassToRemove = PassList[PassIndex];
+				if (IsValid(PassToRemove))
+				{
+					PassToRemove->Modify();
+					PassToRemove->Rename(nullptr, GetTransientPackage(), REN_DontCreateRedirectors);
+				}
+
+				PassList.RemoveAt(PassIndex, EAllowShrinking::No);
 			}
 
 			PassList.Shrink();
