Composite: Fixed incorrect alpha threshold in keyer due to already premultiplied colors.

* Despill is no longer operating on premult colors as well, which is preferable despite not operating on luminance.

#prelight 68e56c110341a491fbb48ef8
#rnx

diff --git a/Engine/Plugins/Compositing/Composite/Shaders/Private/CompositeColorKeyer.usf b/Engine/Plugins/Compositing/Composite/Shaders/Private/CompositeColorKeyer.usf
--- a/Engine/Plugins/Compositing/Composite/Shaders/Private/CompositeColorKeyer.usf
+++ b/Engine/Plugins/Compositing/Composite/Shaders/Private/CompositeColorKeyer.usf
@@ -36,8 +36,8 @@ float2 UVToScreenPos(float2 UV, float2 Extent)
 // Output color function with debug convenience
 float4 Output(float4 Color)
 {
-	// Optionally invert alpha
-	Color.a = select(bInvertAlpha > 0, 1.0f - Color.a, Color.a);
+	// Saturate and optionally invert alpha
+	Color.a = saturate(select(bInvertAlpha > 0, 1.0f - Color.a, Color.a));
 
 	if(Visualization == 1)
 	{
@@ -51,7 +51,8 @@ float4 Output(float4 Color)
 		Color.a = 1.0f;
 	}
 
-	return Color;
+	// Keyed output colors should always be premultiplied
+	return Premultiply(Color);
 }
 
 [numthreads(TILE_SIZE, TILE_SIZE, 1)]
diff --git a/Engine/Plugins/Compositing/Composite/Shaders/Private/CompositeColorKeyerCommon.ush b/Engine/Plugins/Compositing/Composite/Shaders/Private/CompositeColorKeyerCommon.ush
--- a/Engine/Plugins/Compositing/Composite/Shaders/Private/CompositeColorKeyerCommon.ush
+++ b/Engine/Plugins/Compositing/Composite/Shaders/Private/CompositeColorKeyerCommon.ush
@@ -20,6 +20,16 @@ float SafeRcp(float x)
 	return x > 0.0 ? rcp(x) : 0.0;
 }
 
+float4 Unpremultiply(float4 InColor)
+{
+	return float4(InColor.xyz * SafeRcp(InColor.w), InColor.w);
+}
+
+float4 Premultiply(float4 InColor)
+{
+	return float4(InColor.xyz * InColor.w, InColor.w);
+}
+
 float LuminanceSRGB(float3 InColor)
 {
 	return dot(InColor, float3(0.2126390059, 0.7151686788, 0.0721923154));
@@ -213,7 +223,9 @@ float4 GreenScreenKeyer(
 	// Clamp alpha
 	Alpha = saturate(Alpha);
 
-	// Note: Output is already considered pre-multiplied
+	// Fill is already considered premultiplied, so we un-premultiply it for subsequent color operations.
+	Fill *= SafeRcp(Alpha);
+
 	return float4(Fill, Alpha);
 }
 
