Composite: Fixed composite mesh edits still altering holdout state despite not actively rendering.

* We also fix the tick condition to remove the redundant CDO check, while adding one for the outer object.
* Further UpdateCompositeMeshes call cleanup is delayed to 5.8 to minimize risk.

#rnx

diff --git a/Engine/Plugins/Compositing/Composite/Source/Composite/Private/Layers/CompositeLayerPlate.cpp b/Engine/Plugins/Compositing/Composite/Source/Composite/Private/Layers/CompositeLayerPlate.cpp
--- a/Engine/Plugins/Compositing/Composite/Source/Composite/Private/Layers/CompositeLayerPlate.cpp
+++ b/Engine/Plugins/Compositing/Composite/Source/Composite/Private/Layers/CompositeLayerPlate.cpp
@@ -106,14 +106,7 @@ void UCompositeLayerPlate::SetEnabled(bool bInEnabled)
 
 void UCompositeLayerPlate::Tick(float DeltaTime)
 {
-	if(!bIsEnabled)
-	{
-		return;
-	}
-
-	const ACompositeActor* CompositeActor = GetTypedOuter<ACompositeActor>();
-	
-	if (!IsValid(CompositeActor) || !CompositeActor->IsRendering())
+	if (!IsRendering())
 	{
 		return;
 	}
@@ -171,7 +164,10 @@ ETickableTickType UCompositeLayerPlate::GetTickableTickType() const
 
 bool UCompositeLayerPlate::IsTickable() const
 {
-	return !HasAnyFlags(RF_ClassDefaultObject) && GetTypedOuter<ACompositeActor>();
+	ACompositeActor* CompositeActor = GetTypedOuter<ACompositeActor>();
+
+	// Only tick when instances are registered to a non-CDO composite actor.
+	return IsValid(CompositeActor) && !CompositeActor->HasAnyFlags(RF_ClassDefaultObject);
 }
 
 UWorld* UCompositeLayerPlate::GetTickableGameObjectWorld() const
@@ -341,7 +337,7 @@ void UCompositeLayerPlate::PostEditChangeProperty(FPropertyChangedEvent& Propert
 		}
 
 		UCompositeCoreSubsystem* Subsystem = UWorld::GetSubsystem<UCompositeCoreSubsystem>(GetWorld());
-		if (IsValid(Subsystem))
+		if (IsValid(Subsystem) && IsRendering())
 		{
 			// Register primitives to be marked as holdout and rendered in the built-in composite custom render pass.
 			Subsystem->RegisterPrimitives(GetPrimitives());
@@ -384,7 +380,7 @@ void UCompositeLayerPlate::PostTransacted(const FTransactionObjectEvent& InTrans
 			UpdateCompositeMeshes();
 
 			UCompositeCoreSubsystem* Subsystem = UWorld::GetSubsystem<UCompositeCoreSubsystem>(GetWorld());
-			if (IsValid(Subsystem))
+			if (IsValid(Subsystem) && IsRendering())
 			{
 				// Register primitives to be marked as holdout and rendered in the built-in composite custom render pass.
 				Subsystem->RegisterPrimitives(GetPrimitives());
@@ -638,6 +634,18 @@ void UCompositeLayerPlate::TryCloseMediaProfileSource()
 	ActiveMediaProfile->GetPlaybackManager()->CloseSourcesForConsumer(Args);
 }
 
+bool UCompositeLayerPlate::IsRendering() const
+{
+	if (IsEnabled())
+	{
+		const ACompositeActor* CompositeActor = GetTypedOuter<ACompositeActor>();
+
+		return IsValid(CompositeActor) && CompositeActor->IsRendering();
+	}
+
+	return false;
+}
+
 #if WITH_EDITOR
 void UCompositeLayerPlate::OnPreSaveWorld(UWorld* InWorld, FObjectPreSaveContext ObjectSaveContext)
 {
@@ -738,10 +746,8 @@ void UCompositeLayerPlate::PropagateStateChange(bool bApply, const UWorld* Sourc
 	}
 
 	UCompositeCoreSubsystem* Subsystem = UWorld::GetSubsystem<UCompositeCoreSubsystem>(SourceWorld ? SourceWorld : GetWorld());
-	const ACompositeActor* CompositeActor = GetTypedOuter<ACompositeActor>();
-	const bool bIsRendering = IsValid(CompositeActor) ? CompositeActor->IsRendering() : false;
 
-	if (bApply && IsEnabled() && bIsRendering)
+	if (bApply && IsRendering())
 	{
 		constexpr bool bRegisterOnly = true;
 		UpdateCompositeMeshes(bRegisterOnly);
diff --git a/Engine/Plugins/Compositing/Composite/Source/Composite/Public/Layers/CompositeLayerPlate.h b/Engine/Plugins/Compositing/Composite/Source/Composite/Public/Layers/CompositeLayerPlate.h
--- a/Engine/Plugins/Compositing/Composite/Source/Composite/Public/Layers/CompositeLayerPlate.h
+++ b/Engine/Plugins/Compositing/Composite/Source/Composite/Public/Layers/CompositeLayerPlate.h
@@ -176,6 +176,9 @@ private:
 	/** Attempt to close any media profile sources being played for this plate */
 	void TryCloseMediaProfileSource();
 	
+	/** Returns true if the plate layer is actively rendering. */
+	bool IsRendering() const;
+	
 private:
 	/** Convenience function to find the last valid pass. */
 	int32 FindLastValidPassIndex(TArrayView<const TObjectPtr<UCompositePassBase>> InPasses) const;
