Composure: Fixed issue where deleting was not enabled for composite actors in Composure panel, leading to some unintuitive outcomes when pressing the delete key while a composite actor is selected

diff --git a/Engine/Plugins/Compositing/Composite/Source/CompositeEditor/Private/UI/SCompositePanelLayerTree.cpp b/Engine/Plugins/Compositing/Composite/Source/CompositeEditor/Private/UI/SCompositePanelLayerTree.cpp
--- a/Engine/Plugins/Compositing/Composite/Source/CompositeEditor/Private/UI/SCompositePanelLayerTree.cpp
+++ b/Engine/Plugins/Compositing/Composite/Source/CompositeEditor/Private/UI/SCompositePanelLayerTree.cpp
@@ -2257,6 +2257,20 @@ void SCompositePanelLayerTree::DeleteSelectedItems()
 {
 	TArray<FCompositeActorTreeItemPtr> SelectedItems = TreeView.IsValid() ? TreeView->GetSelectedItems() : TArray<FCompositeActorTreeItemPtr>();
 
+	TSet<ACompositeActor*> ActorsToDelete;
+	for (const FCompositeActorTreeItemPtr& Item : SelectedItems)
+	{
+		// Skip any invalid items, or any items that represent layers, as we want to find items that represent composite actors
+		if (!Item.IsValid() ||
+			!Item->HasValidCompositeActor() ||
+			Item->HasValidCompositeLayer())
+		{
+			continue;
+		}
+
+		ActorsToDelete.Add(Item->CompositeActor.Get());
+	}
+	
 	TMap<ACompositeActor*, TArray<int32>> LayersToDelete;
 	for (const FCompositeActorTreeItemPtr& Item : SelectedItems)
 	{
@@ -2266,6 +2280,12 @@ void SCompositePanelLayerTree::DeleteSelectedItems()
 			continue;
 		}
 
+		// Skip any layers whose owning actor is already getting deleted
+		if (ActorsToDelete.Contains(Item->CompositeActor.Get()))
+		{
+			continue;
+		}
+
 		if (LayersToDelete.Contains(Item->CompositeActor.Get()))
 		{
 			LayersToDelete[Item->CompositeActor.Get()].Add(Item->LayerIndex);
@@ -2276,9 +2296,26 @@ void SCompositePanelLayerTree::DeleteSelectedItems()
 		}
 	}
 
+	if (ActorsToDelete.IsEmpty() && LayersToDelete.IsEmpty())
+	{
+		return;
+	}
+
+	FScopedTransaction DeleteActorTransaction(LOCTEXT("DeleteTransaction", "Delete Actors/Layers"));
+	
+	if (ActorsToDelete.Num())
+	{
+		for (ACompositeActor* Actor : ActorsToDelete)
+		{
+			if (IsValid(Actor))
+			{
+				Actor->Destroy();
+			}
+		}
+	}
+	
 	if (LayersToDelete.Num())
 	{
-		FScopedTransaction DeleteLayerTransaction(LOCTEXT("DeleteLayerTransaction", "Delete Layer"));
 		for (TPair<ACompositeActor*, TArray<int32>>& CompositeActorLayerList : LayersToDelete)
 		{
 			ACompositeActor* CompositeActor = CompositeActorLayerList.Key;
@@ -2320,12 +2357,22 @@ void SCompositePanelLayerTree::DeleteSelectedItems()
 				return InTreeItem->CompositeActor.Get() == CompositeActor;
 			});
 
-			if (TreeItemIndex != INDEX_NONE)
+			if (TreeItemIndex != INDEX_NONE && ActorsToDelete.IsEmpty())
 			{
 				FillCompositeActorTreeItem(CompositeActorTreeItems[TreeItemIndex]);
 			}
 		}
+	}
 
+	if (!ActorsToDelete.IsEmpty())
+	{
+		// If an actor was deleted, rebuild the whole tree view
+		constexpr bool bPreserveSelection = false;
+		FillCompositeActorTree(bPreserveSelection);
+	}
+	else if (!LayersToDelete.IsEmpty())
+	{
+		// Otherwise, we can just refresh the tree view since the layers were updated with FillCompositeActorTreeItem
 		RefreshAndExpandTreeView();
 	}
 }
@@ -2333,12 +2380,7 @@ void SCompositePanelLayerTree::DeleteSelectedItems()
 bool SCompositePanelLayerTree::CanDeleteSelectedItems()
 {
 	TArray<FCompositeActorTreeItemPtr> SelectedItems = TreeView.IsValid() ? TreeView->GetSelectedItems() : TArray<FCompositeActorTreeItemPtr>();
-	const bool bContainsLayerInSelection = SelectedItems.ContainsByPredicate([](const FCompositeActorTreeItemPtr& InTreeItem)
-	{
-		return InTreeItem.IsValid() && InTreeItem->HasValidCompositeLayer();
-	});
-	
-	return bContainsLayerInSelection;
+	return !SelectedItems.IsEmpty();
 }
 
 void SCompositePanelLayerTree::RenameSelectedItem()
