Composure: Fixed issue where placing a composite mesh actor from Composure panel would move the currently piloted actor.

diff --git a/Engine/Plugins/Compositing/Composite/Source/CompositeEditor/Private/Customizations/CompositeLayerPlateCustomization.cpp b/Engine/Plugins/Compositing/Composite/Source/CompositeEditor/Private/Customizations/CompositeLayerPlateCustomization.cpp
--- a/Engine/Plugins/Compositing/Composite/Source/CompositeEditor/Private/Customizations/CompositeLayerPlateCustomization.cpp
+++ b/Engine/Plugins/Compositing/Composite/Source/CompositeEditor/Private/Customizations/CompositeLayerPlateCustomization.cpp
@@ -12,6 +12,7 @@
 #include "Engine/Texture2D.h"
 #include "IDetailGroup.h"
 #include "IPropertyUtilities.h"
+#include "LevelEditorSubsystem.h"
 #include "MediaSource.h"
 #include "MediaTexture.h"
 #include "PropertyCustomizationHelpers.h"
@@ -353,7 +354,19 @@ void FCompositeLayerPlateCustomization::CreateCompositeMeshActor()
 		{
 			GEditor->SelectNone(true, true);
 			GEditor->SelectActor(AddedMeshActor, true, true);
-			GUnrealEd->Exec(CompositeActor->GetWorld(), TEXT("CAMERA ALIGN ACTIVEVIEWPORTONLY"));
+
+			bool bAlignViewport = true;
+			if (ULevelEditorSubsystem* LevelEditorSubsystem = GEditor->GetEditorSubsystem<ULevelEditorSubsystem>())
+			{
+				// Don't align the viewport if there is an active pilot actor, as that will cause the actor to be moved when the viewport
+				// camera is moved
+				bAlignViewport = !IsValid(LevelEditorSubsystem->GetPilotLevelActor());
+			}
+
+			if (bAlignViewport)
+			{
+				GUnrealEd->Exec(CompositeActor->GetWorld(), TEXT("CAMERA ALIGN ACTIVEVIEWPORTONLY"));
+			}
 		}
 	}
 }
