Composure: Changed logic for when Pilot button is enabled to be only for the active composite actor. Also added functionality to automatically stop piloting when a change to the active composite actor or its camera is detected.

diff --git a/Engine/Plugins/Compositing/Composite/Source/CompositeEditor/Private/UI/SCompositePanelLayerTree.cpp b/Engine/Plugins/Compositing/Composite/Source/CompositeEditor/Private/UI/SCompositePanelLayerTree.cpp
--- a/Engine/Plugins/Compositing/Composite/Source/CompositeEditor/Private/UI/SCompositePanelLayerTree.cpp
+++ b/Engine/Plugins/Compositing/Composite/Source/CompositeEditor/Private/UI/SCompositePanelLayerTree.cpp
@@ -310,6 +310,7 @@ class SCompositePanelLayerTreeItemRow : public SMultiColumnTableRow<SCompositePa
 	using FTreeItemWeakPtr = TWeakPtr<FTreeItem>;
 	
 public:
+	DECLARE_DELEGATE_OneParam(FOnCompositeActorActivated, const TStrongObjectPtr<ACompositeActor>&)
 	DECLARE_DELEGATE_RetVal_OneParam(FReply, FOnLayerSoloToggled, const FTreeItemPtr&)
 	DECLARE_DELEGATE_TwoParams(FOnLayerMoved, const FTreeItemPtr& /* InItemToMove */, int32 /* InDestIndex */)
 
@@ -348,6 +349,7 @@ private:
 	
 public:
 	SLATE_BEGIN_ARGS(SCompositePanelLayerTreeItemRow) { }
+		SLATE_EVENT(FOnCompositeActorActivated, OnCompositeActorActivated)
 		SLATE_EVENT(FOnLayerSoloToggled, OnLayerSoloToggled)
 		SLATE_EVENT(FOnLayerMoved, OnLayerMoved)
 	SLATE_END_ARGS()
@@ -355,6 +357,7 @@ public:
 	void Construct(const FArguments& InArgs, const TSharedRef<STableViewBase>& InOwnerTable, FTreeItemPtr InTreeItem)
 	{
 		TreeItem = InTreeItem;
+		OnCompositeActorActivated = InArgs._OnCompositeActorActivated;
 		OnLayerSoloToggled = InArgs._OnLayerSoloToggled;
 		OnLayerMoved = InArgs._OnLayerMoved;
 
@@ -508,6 +511,7 @@ private:
 		{
 			// No transactions, the property is non-transactional
 			CompositeActor->SetActive(InCheckBoxState == ECheckBoxState::Checked);
+			OnCompositeActorActivated.ExecuteIfBound(CompositeActor);
 		}
 	}
 
@@ -798,6 +802,9 @@ private:
 
 	/** Image widget used for the 'Solo' option button */
 	TSharedPtr<SImage> SoloButtonImage;
+
+	/** Raised when a composite actor is set to the active composite actor */
+	FOnCompositeActorActivated OnCompositeActorActivated;
 	
 	/** Callback that is raised when the layer's "solo" flag is enabled */
 	FOnLayerSoloToggled OnLayerSoloToggled;
@@ -976,6 +983,7 @@ void SCompositePanelLayerTree::Construct(const FArguments& InArgs)
 				.OnGenerateRow_Lambda([this](FCompositeActorTreeItemPtr InTreeItem, const TSharedRef<STableViewBase>& InOwnerTable)
 				{
 					return SNew(SCompositePanelLayerTreeItemRow, InOwnerTable, InTreeItem)
+						.OnCompositeActorActivated(this, &SCompositePanelLayerTree::OnCompositeActorActivated)
 						.OnLayerSoloToggled(this, &SCompositePanelLayerTree::OnLayerSoloToggled)
 						.OnLayerMoved(this, &SCompositePanelLayerTree::OnLayerMoved);
 				})
@@ -1038,6 +1046,23 @@ SCompositePanelLayerTree::~SCompositePanelLayerTree()
 	}
 }
 
+void SCompositePanelLayerTree::Tick(const FGeometry& AllottedGeometry, const double InCurrentTime, const float InDeltaTime)
+{
+	SCompoundWidget::Tick(AllottedGeometry, InCurrentTime, InDeltaTime);
+
+	// Update the IsPilotingCamera flag by checking that the current piloted actor is the active composite actor's camera
+	bCachedIsPilotingCamera = false;
+
+	AActor* CompositeCameraActor = GetActiveCompositeActorCamera();
+	if (IsValid(CompositeCameraActor))
+	{
+		if (ULevelEditorSubsystem* LevelEditorSubsystem = GEditor->GetEditorSubsystem<ULevelEditorSubsystem>())
+		{
+			bCachedIsPilotingCamera = CompositeCameraActor == LevelEditorSubsystem->GetPilotLevelActor();
+		}
+	}
+}
+
 FReply SCompositePanelLayerTree::OnKeyDown(const FGeometry& MyGeometry, const FKeyEvent& InKeyEvent)
 {
 	if (this->CommandList->ProcessCommandBindings(InKeyEvent))
@@ -1052,6 +1077,8 @@ void SCompositePanelLayerTree::PostUndo(bool bSuccess)
 {
 	constexpr bool bPreserveSelection = true;
 	FillCompositeActorTree(bPreserveSelection);
+
+	VerifyActiveCompositeActorPiloting();
 }
 
 void SCompositePanelLayerTree::PostRedo(bool bSuccess)
@@ -1355,98 +1382,128 @@ bool SCompositePanelLayerTree::CanAddLayer() const
 
 bool SCompositePanelLayerTree::CanPilotCamera() const
 {
-	if (CompositeActorTreeItems.Num() == 1)
+	// Always allow the user to eject regardless of selection if a composite actor's camera is being piloted
+	if (bCachedIsPilotingCamera)
 	{
-		return CompositeActorTreeItems[0]->CompositeActor.IsValid();
+		return true;
 	}
 
 	TArray<FCompositeActorTreeItemPtr> SelectedItems = TreeView.IsValid() ? TreeView->GetSelectedItems() : TArray<FCompositeActorTreeItemPtr>();
-	if (!SelectedItems.Num())
+	if (!SelectedItems.IsEmpty())
 	{
-		return false;
-	}
+		// If the active composite actor is not included in the currently selected items, disable the Pilot button
+		const bool bSelectionContainsActiveCompositeActor = SelectedItems.ContainsByPredicate([](const FCompositeActorTreeItemPtr& InTreeItem)
+		{
+			return InTreeItem.IsValid() && InTreeItem->HasValidCompositeActor() && InTreeItem->CompositeActor->IsActive();
+		});
 
-	if (!SelectedItems[0]->HasValidCompositeActor())
-	{
-		return false;
+		if (!bSelectionContainsActiveCompositeActor)
+		{
+			return false;
+		}
 	}
-
-	return SelectedItems[0]->CompositeActor->GetCamera().OtherActor.IsValid();
+	
+	AActor* CompositeCameraActor = GetActiveCompositeActorCamera();
+	return IsValid(CompositeCameraActor);
 }
 
 bool SCompositePanelLayerTree::IsPilotingCamera() const
 {
-	TWeakObjectPtr<ACompositeActor> CompositeActorWeak = nullptr;
+	return bCachedIsPilotingCamera;
+}
 
-	if (CompositeActorTreeItems.Num() == 1)
-	{
-		CompositeActorWeak = CompositeActorTreeItems[0]->CompositeActor;
-	}
-	else
+void SCompositePanelLayerTree::OnPilotCameraToggled(bool bPilotCamera)
+{
+	if (bPilotCamera)
 	{
 		TArray<FCompositeActorTreeItemPtr> SelectedItems = TreeView.IsValid() ? TreeView->GetSelectedItems() : TArray<FCompositeActorTreeItemPtr>();
 		if (!SelectedItems.IsEmpty())
 		{
-			CompositeActorWeak = SelectedItems[0]->CompositeActor;
+			// If the active composite actor is not included in the currently selected items, don't try and pilot anything
+			const bool bSelectionContainsActiveCompositeActor = SelectedItems.ContainsByPredicate([](const FCompositeActorTreeItemPtr& InTreeItem)
+			{
+				return InTreeItem.IsValid() && InTreeItem->HasValidCompositeActor() && InTreeItem->CompositeActor->IsActive();
+			});
+
+			if (!bSelectionContainsActiveCompositeActor)
+			{
+				return;
+			}
 		}
-	}
 
-	TStrongObjectPtr CompositeActor = CompositeActorWeak.Pin();
+		AActor* CompositeCameraActor = GetActiveCompositeActorCamera();
 
-	if (!CompositeActor.IsValid())
+		if (!IsValid(CompositeCameraActor))
+		{
+			return;
+		}
+		
+		PilotCamera(CompositeCameraActor);
+	}
+	else
 	{
-		return false;
+		StopPilotingCamera();
 	}
+}
 
-	AActor* CompositeCameraActor = CompositeActor->GetCamera().OtherActor.Get();
-
-	if (!IsValid(CompositeCameraActor))
+AActor* SCompositePanelLayerTree::GetActiveCompositeActorCamera() const
+{
+	TStrongObjectPtr<ACompositeActor> ActiveCompositeActor = nullptr;
+	const FCompositeActorTreeItemPtr* ActiveTreeItem = CompositeActorTreeItems.FindByPredicate([](const FCompositeActorTreeItemPtr& InTreeItem)
 	{
-		return false;
-	}
+		return InTreeItem.IsValid() && InTreeItem->HasValidCompositeActor() && InTreeItem->CompositeActor->IsActive();
+	});
 	
-	if (ULevelEditorSubsystem* LevelEditorSubsystem = GEditor->GetEditorSubsystem<ULevelEditorSubsystem>())
+	if (ActiveTreeItem)
 	{
-		return CompositeCameraActor == LevelEditorSubsystem->GetPilotLevelActor();
+		ActiveCompositeActor = (*ActiveTreeItem)->CompositeActor.Pin();
 	}
 
-	return false;
+	if (ActiveCompositeActor.IsValid())
+	{
+		return ActiveCompositeActor->GetCamera().OtherActor.Get();
+	}
+
+	return nullptr;
 }
 
-void SCompositePanelLayerTree::OnPilotCameraToggled(bool bPilotCamera)
+void SCompositePanelLayerTree::PilotCamera(AActor* InCameraActorToPilot)
 {
-	TWeakObjectPtr<ACompositeActor> CompositeActorWeak = nullptr;
-
-	if (CompositeActorTreeItems.Num() == 1)
-	{
-		CompositeActorWeak = CompositeActorTreeItems[0]->CompositeActor;
-	}
-	else
+	if (ULevelEditorSubsystem* LevelEditorSubsystem = GEditor->GetEditorSubsystem<ULevelEditorSubsystem>())
 	{
-		TArray<FCompositeActorTreeItemPtr> SelectedItems = TreeView.IsValid() ? TreeView->GetSelectedItems() : TArray<FCompositeActorTreeItemPtr>();
-		if (!SelectedItems.IsEmpty())
+		FLevelEditorModule& LevelEditorModule = FModuleManager::GetModuleChecked<FLevelEditorModule>("LevelEditor");
+		TSharedPtr<SLevelViewport> ActiveLevelViewport = nullptr;
+		TSharedPtr<ILevelEditor> LevelEditor = LevelEditorModule.GetFirstLevelEditor();
+		if (LevelEditor.IsValid())
 		{
-			CompositeActorWeak = SelectedItems[0]->CompositeActor;
+			ActiveLevelViewport = LevelEditor->GetActiveViewportInterface();
 		}
-	}
 
-	TStrongObjectPtr CompositeActor = CompositeActorWeak.Pin();
+		// The subsystem currently has no method for capturing the view transform, so we extract it ourselves for now.
+		if (ActiveLevelViewport.IsValid())
+		{
+			FLevelEditorViewportClient& Client = ActiveLevelViewport->GetLevelViewportClient();
+			CachedPerspectiveCameraTransform = Client.GetViewTransform();
+			CachedViewFOV = Client.ViewFOV;
+		}
 
-	if (!CompositeActor.IsValid())
-	{
-		return;
+		LevelEditorSubsystem->PilotLevelActor(InCameraActorToPilot);
+		LevelEditorSubsystem->SetExactCameraView(true);
 	}
+}
 
-	AActor* CompositeCameraActor = CompositeActor->GetCamera().OtherActor.Get();
-
-	if (!IsValid(CompositeCameraActor))
-	{
-		return;
-	}
-	
+void SCompositePanelLayerTree::StopPilotingCamera()
+{
 	if (ULevelEditorSubsystem* LevelEditorSubsystem = GEditor->GetEditorSubsystem<ULevelEditorSubsystem>())
 	{
-		// The subsytem currently has no method for capturing the view transform, so we extract it ourselves for now.
+		// Early abort if the level editor is not actively piloting anything
+		if (!LevelEditorSubsystem->GetPilotLevelActor())
+		{
+			return;
+		}
+		
+		LevelEditorSubsystem->EjectPilotLevelActor();
+
 		FLevelEditorModule& LevelEditorModule = FModuleManager::GetModuleChecked<FLevelEditorModule>("LevelEditor");
 		TSharedPtr<SLevelViewport> ActiveLevelViewport = nullptr;
 		TSharedPtr<ILevelEditor> LevelEditor = LevelEditorModule.GetFirstLevelEditor();
@@ -1454,39 +1511,44 @@ void SCompositePanelLayerTree::OnPilotCameraToggled(bool bPilotCamera)
 		{
 			ActiveLevelViewport = LevelEditor->GetActiveViewportInterface();
 		}
-
-		if (bPilotCamera)
+		
+		if (ActiveLevelViewport.IsValid())
 		{
-			if (ActiveLevelViewport.IsValid())
+			// Restore the previously cached view transform
+			FLevelEditorViewportClient& Client = ActiveLevelViewport->GetLevelViewportClient();
+			if (CachedPerspectiveCameraTransform.IsSet())
 			{
-				FLevelEditorViewportClient& Client = ActiveLevelViewport->GetLevelViewportClient();
-				CachedPerspectiveCameraTransform = Client.GetViewTransform();
-				CachedViewFOV = Client.ViewFOV;
+				Client.SetViewLocation(CachedPerspectiveCameraTransform.GetValue().GetLocation());
+				Client.SetViewRotation(CachedPerspectiveCameraTransform.GetValue().GetRotation());
 			}
 
-			LevelEditorSubsystem->PilotLevelActor(CompositeCameraActor);
-			LevelEditorSubsystem->SetExactCameraView(true);
+			if(CachedViewFOV.IsSet())
+			{
+				Client.ViewFOV = CachedViewFOV.GetValue();
+			}
 		}
-		else
-		{
-			LevelEditorSubsystem->EjectPilotLevelActor();
+	}
+}
 
-			if (ActiveLevelViewport.IsValid())
+void SCompositePanelLayerTree::VerifyActiveCompositeActorPiloting()
+{
+	if (bCachedIsPilotingCamera)
+	{
+		// Check that the active composite actor's camera is still the one being piloted, and if it isn't, stop piloting
+		bool bIsStillPilotingCamera = false;
+		AActor* CompositeCameraActor = GetActiveCompositeActorCamera();
+		if (IsValid(CompositeCameraActor))
+		{
+			if (ULevelEditorSubsystem* LevelEditorSubsystem = GEditor->GetEditorSubsystem<ULevelEditorSubsystem>())
 			{
-				// Restore the previously cached view transform
-				FLevelEditorViewportClient& Client = ActiveLevelViewport->GetLevelViewportClient();
-				if (CachedPerspectiveCameraTransform.IsSet())
-				{
-					Client.SetViewLocation(CachedPerspectiveCameraTransform.GetValue().GetLocation());
-					Client.SetViewRotation(CachedPerspectiveCameraTransform.GetValue().GetRotation());
-				}
-
-				if(CachedViewFOV.IsSet())
-				{
-					Client.ViewFOV = CachedViewFOV.GetValue();
-				}
+				bIsStillPilotingCamera = CompositeCameraActor == LevelEditorSubsystem->GetPilotLevelActor();
 			}
 		}
+
+		if (!bIsStillPilotingCamera)
+		{
+			StopPilotingCamera();
+		}
 	}
 }
 
@@ -1522,9 +1584,11 @@ void SCompositePanelLayerTree::OnGlobalActiveCheckStateChanged(ECheckBoxState Ch
 				continue;
 			}
 
-			if (TreeItem->CompositeActor.Get() == LastActiveActor.Get())
+			TStrongObjectPtr<ACompositeActor> CompositeActor = TreeItem->CompositeActor.Pin();
+			if (CompositeActor.Get() == LastActiveActor.Get())
 			{
-				TreeItem->CompositeActor->SetActive(true);
+				CompositeActor->SetActive(true);
+				OnCompositeActorActivated(CompositeActor);
 				bActivated = true;
 			}
 		}
@@ -1538,8 +1602,10 @@ void SCompositePanelLayerTree::OnGlobalActiveCheckStateChanged(ECheckBoxState Ch
 				{
 					continue;
 				}
-				
-				TreeItem->CompositeActor->SetActive(true);
+
+				TStrongObjectPtr<ACompositeActor> CompositeActor = TreeItem->CompositeActor.Pin();
+				CompositeActor->SetActive(true);
+				OnCompositeActorActivated(CompositeActor);
 				break;
 			}
 		}
@@ -1559,6 +1625,7 @@ void SCompositePanelLayerTree::OnGlobalActiveCheckStateChanged(ECheckBoxState Ch
 			}
 
 			TreeItem->CompositeActor->SetActive(false);
+			OnCompositeActorActivated(nullptr);
 		}
 	}
 }
@@ -1886,6 +1953,31 @@ void SCompositePanelLayerTree::OnLayerAdded(const UClass* InLayerClass)
 	RefreshCompositeActorLayers(CompositeActor.Get());
 }
 
+void SCompositePanelLayerTree::OnCompositeActorActivated(const TStrongObjectPtr<ACompositeActor>& CompositeActor)
+{
+	if (!CompositeActor.IsValid())
+	{
+		StopPilotingCamera();
+		return;
+	}
+
+	AActor* CompositeCameraActor = CompositeActor->GetCamera().OtherActor.Get();
+
+	if (!IsValid(CompositeCameraActor))
+	{
+		StopPilotingCamera();
+		return;
+	}
+	
+	if (ULevelEditorSubsystem* LevelEditorSubsystem = GEditor->GetEditorSubsystem<ULevelEditorSubsystem>())
+	{
+		if (CompositeCameraActor != LevelEditorSubsystem->GetPilotLevelActor())
+		{
+			StopPilotingCamera();
+		}
+	}
+}
+
 FReply SCompositePanelLayerTree::OnLayerSoloToggled(const FCompositeActorTreeItemPtr& InLayerTreeItem)
 {
 	if (!InLayerTreeItem.IsValid() ||
@@ -1992,6 +2084,7 @@ void SCompositePanelLayerTree::OnActorRemovedFromLevel(AActor* Actor)
 	if (Actor->IsA<ACompositeActor>())
 	{
 		FillCompositeActorTree();
+		VerifyActiveCompositeActorPiloting();
 	}
 }
 
@@ -2029,16 +2122,40 @@ void SCompositePanelLayerTree::OnObjectReplaced(const TMap<UObject*, UObject*>&
 	if (bContainsCompositeActor)
 	{
 		FillCompositeActorTree();
+		VerifyActiveCompositeActorPiloting();
 	}
 }
 
 void SCompositePanelLayerTree::OnObjectPropertyChanged(UObject* Object, FPropertyChangedEvent& PropertyChangedEvent)
 {
-	if (Object->IsA<ACompositeActor>() && PropertyChangedEvent.GetPropertyName() == GET_MEMBER_NAME_CHECKED(ACompositeActor, CompositeLayers))
+	if (Object->IsA<ACompositeActor>())
 	{
-		constexpr bool bRefreshTree = true;
-		constexpr bool bPreserveSelection = true;
-		RefreshCompositeActorLayers(Cast<ACompositeActor>(Object), bRefreshTree, bPreserveSelection);
+		const FName PropertyName = PropertyChangedEvent.GetPropertyName();
+		if (PropertyName == GET_MEMBER_NAME_CHECKED(ACompositeActor, CompositeLayers))
+		{
+			constexpr bool bRefreshTree = true;
+			constexpr bool bPreserveSelection = true;
+			RefreshCompositeActorLayers(Cast<ACompositeActor>(Object), bRefreshTree, bPreserveSelection);
+		}
+		else if (PropertyName == GET_MEMBER_NAME_CHECKED(ACompositeActor, Camera))
+		{
+			ACompositeActor* CompositeActor = Cast<ACompositeActor>(Object);
+			AActor* CompositeCameraActor = CompositeActor->GetCamera().OtherActor.Get();
+
+			if (!IsValid(CompositeCameraActor))
+			{
+				StopPilotingCamera();
+				return;
+			}
+	
+			if (ULevelEditorSubsystem* LevelEditorSubsystem = GEditor->GetEditorSubsystem<ULevelEditorSubsystem>())
+			{
+				if (CompositeCameraActor != LevelEditorSubsystem->GetPilotLevelActor())
+				{
+					StopPilotingCamera();
+				}
+			}
+		}
 	}
 }
 
diff --git a/Engine/Plugins/Compositing/Composite/Source/CompositeEditor/Private/UI/SCompositePanelLayerTree.h b/Engine/Plugins/Compositing/Composite/Source/CompositeEditor/Private/UI/SCompositePanelLayerTree.h
--- a/Engine/Plugins/Compositing/Composite/Source/CompositeEditor/Private/UI/SCompositePanelLayerTree.h
+++ b/Engine/Plugins/Compositing/Composite/Source/CompositeEditor/Private/UI/SCompositePanelLayerTree.h
@@ -59,6 +59,7 @@ public:
 	virtual ~SCompositePanelLayerTree() override;
 
 	// SWidget interface
+	virtual void Tick(const FGeometry& AllottedGeometry, const double InCurrentTime, const float InDeltaTime) override;
 	virtual FReply OnKeyDown(const FGeometry& MyGeometry, const FKeyEvent& InKeyEvent) override;
 	// End of SWidget interface
 	
@@ -104,6 +105,18 @@ private:
 	/** Sets the piloted actor to the currently selected composite actor camera, or clears it if it is already being piloted */
 	void OnPilotCameraToggled(bool bPilotCamera);
 
+	/** Gets the currently active composite actor's camera, or nullptr if no active composite actor or no valid camera was found */
+	AActor* GetActiveCompositeActorCamera() const;
+
+	/** Starts piloting the specified camera actor */
+	void PilotCamera(AActor* InCameraActorToPilot);
+
+	/** Stops piloting the currently piloted camera actor */
+	void StopPilotingCamera();
+
+	/** Checks whether the current active composite actor's camera is still the actor being piloted, and if not, stops piloting */
+	void VerifyActiveCompositeActorPiloting();
+	
 	/** Gets the global active checkbox state. */
 	ECheckBoxState GetGlobalActiveCheckState() const;
 
@@ -137,6 +150,9 @@ private:
 	/** Raised when a new layer is being added to the selected composite actor */
 	void OnLayerAdded(const UClass* InLayerClass);
 
+	/** Raised when a composite actor is made the composite actor */
+	void OnCompositeActorActivated(const TStrongObjectPtr<ACompositeActor>& CompositeActor);
+	
 	/** Raised when the Solo button has been toggled on the specified layer tree item */
 	FReply OnLayerSoloToggled(const FCompositeActorTreeItemPtr& InLayerTreeItem);
 
@@ -206,6 +222,9 @@ private:
 
 	/** When true, selection changes in the tree view will not invoke the SelectionChanged delegate */
 	bool bSilenceSelectionChanges = false;
+
+	/** Stores whether the active composite actor's camera is currently being piloted */
+	bool bCachedIsPilotingCamera = false;
 	
 	friend class SCompositePanelLayerTreeItemRow;
 	friend class SCompositePanelLayerTreeToolbar;
