Composite: Fixed undo failing to restore deleted default layers.

* NewObject is preferable over DefaultSubObject for pass instances: RF_DefaultSubObject are considered templates, and are not restored with undo after deletion.
* We now create the default layer setup in a custom actor factory.
* The constructor removal of default pass subobjects does not break existing content.

#rnx

diff --git a/Engine/Plugins/Compositing/Composite/Source/Composite/Private/CompositeActor.cpp b/Engine/Plugins/Compositing/Composite/Source/Composite/Private/CompositeActor.cpp
--- a/Engine/Plugins/Compositing/Composite/Source/Composite/Private/CompositeActor.cpp
+++ b/Engine/Plugins/Compositing/Composite/Source/Composite/Private/CompositeActor.cpp
@@ -268,23 +268,6 @@ ACompositeActor::ACompositeActor(const FObjectInitializer& ObjectInitializer)
 
 	ViewProjectionComponent = CreateDefaultSubobject<UCompositeViewProjectionComponent>(TEXT("CameraViewProjectionComponent"));
 	AddInstanceComponent(ViewProjectionComponent);
-
-	// Default layers when creating a new actor
-	CompositeLayers.Reserve(3);
-	CompositeLayers.Add(CreateDefaultSubobject<UCompositeLayerMainRender>(TEXT("MainRenderLayer")));
-	CompositeLayers.Add(CreateDefaultSubobject<UCompositeLayerShadowReflection>(TEXT("ShadowReflectionLayer")));
-	CompositeLayers.Add(CreateDefaultSubobject<UCompositeLayerPlate>(TEXT("PlateLayer")));
-
-#if WITH_EDITOR
-	// Assign default Name value to all layers
-	for (TObjectPtr<UCompositeLayerBase>& Layer : CompositeLayers)
-	{
-		Layer->Name = Layer->GetName();
-	}
-#endif
-	
-	// By default, we disable the shadow reflection catcher due to its significant added cost.
-	CompositeLayers[1]->SetEnabled(false);
 }
 
 ACompositeActor::~ACompositeActor() = default;
diff --git a/Engine/Plugins/Compositing/Composite/Source/Composite/Public/CompositeActor.h b/Engine/Plugins/Compositing/Composite/Source/Composite/Public/CompositeActor.h
--- a/Engine/Plugins/Compositing/Composite/Source/Composite/Public/CompositeActor.h
+++ b/Engine/Plugins/Compositing/Composite/Source/Composite/Public/CompositeActor.h
@@ -210,6 +210,7 @@ private:
 	friend class SCompositePanelLayerTree;
 	friend class FCompositeActorCustomization;
 	friend class FCompositeActorPanelDetailCustomization;
+	friend class UCompositeActorFactory;
 };
 
 #undef UE_API
diff --git a/Engine/Plugins/Compositing/Composite/Source/CompositeEditor/Private/ActorFactories/CompositeActorFactory.cpp b/Engine/Plugins/Compositing/Composite/Source/CompositeEditor/Private/ActorFactories/CompositeActorFactory.cpp
new file mode 100644
--- /dev/null
+++ b/Engine/Plugins/Compositing/Composite/Source/CompositeEditor/Private/ActorFactories/CompositeActorFactory.cpp
@@ -0,0 +1,56 @@
+// Copyright Epic Games, Inc. All Rights Reserved.
+
+#include "CompositeActorFactory.h"
+
+#include "CompositeActor.h"
+#include "Layers/CompositeLayerMainRender.h"
+#include "Layers/CompositeLayerShadowReflection.h"
+#include "Layers/CompositeLayerPlate.h"
+
+#include UE_INLINE_GENERATED_CPP_BY_NAME(CompositeActorFactory)
+
+#define LOCTEXT_NAMESPACE "CompositeActorFactory"
+
+UCompositeActorFactory::UCompositeActorFactory(const FObjectInitializer& ObjInit)
+	: Super(ObjInit)
+{
+	DisplayName = LOCTEXT("DisplayName", "Composite Actor");
+	NewActorClass = ACompositeActor::StaticClass();
+}
+
+void UCompositeActorFactory::PostSpawnActor(UObject* Asset, AActor* NewActor)
+{
+	Super::PostSpawnActor(Asset, NewActor);
+
+	ACompositeActor* CompositeActor = Cast<ACompositeActor>(NewActor);
+
+	// Configure newly spawned composite actor instance
+	if (IsValid(CompositeActor) && !CompositeActor->IsTemplate())
+	{
+		// Default layers when creating a new actor
+		TArray<TObjectPtr<UCompositeLayerBase>> NewLayers;
+		NewLayers.Reserve(3);
+		NewLayers.Add(NewObject<UCompositeLayerMainRender>(CompositeActor, TEXT("MainRenderLayer"), RF_Transactional));
+		NewLayers.Add(NewObject<UCompositeLayerShadowReflection>(CompositeActor, TEXT("ShadowReflectionLayer"), RF_Transactional));
+		NewLayers.Add(NewObject<UCompositeLayerPlate>(CompositeActor, TEXT("PlateLayer"), RF_Transactional));
+
+		// Assign default Name value to all layers
+		for (TObjectPtr<UCompositeLayerBase>& Layer : NewLayers)
+		{
+			Layer->Name = Layer->GetName();
+		}
+
+		// By default, we disable the shadow reflection catcher due to its significant added cost.
+		NewLayers[1]->SetEnabled(false);
+
+		CompositeActor->SetCompositeLayers(NewLayers);
+
+		// Refresh Composure panel
+		FProperty* CompositeLayersProperty = FindFProperty<FProperty>(ACompositeActor::StaticClass(), GET_MEMBER_NAME_CHECKED(ACompositeActor, CompositeLayers));
+		FPropertyChangedEvent LayersChangedEvent(CompositeLayersProperty, EPropertyChangeType::ArrayAdd);
+		FCoreUObjectDelegates::OnObjectPropertyChanged.Broadcast(NewActor, LayersChangedEvent);
+	}
+}
+
+#undef LOCTEXT_NAMESPACE
+
diff --git a/Engine/Plugins/Compositing/Composite/Source/CompositeEditor/Private/ActorFactories/CompositeActorFactory.h b/Engine/Plugins/Compositing/Composite/Source/CompositeEditor/Private/ActorFactories/CompositeActorFactory.h
new file mode 100644
--- /dev/null
+++ b/Engine/Plugins/Compositing/Composite/Source/CompositeEditor/Private/ActorFactories/CompositeActorFactory.h
@@ -0,0 +1,25 @@
+// Copyright Epic Games, Inc. All Rights Reserved.
+
+#pragma once
+
+#include "UObject/ObjectMacros.h"
+#include "ActorFactories/ActorFactory.h"
+
+#include "CompositeActorFactory.generated.h"
+
+UCLASS()
+class UCompositeActorFactory : public UActorFactory
+{
+	GENERATED_BODY()
+
+public:
+	/** Constructor */
+	UCompositeActorFactory(const FObjectInitializer& ObjInit);
+
+public:
+
+	//~ Begin UActorFactory interface
+	virtual void PostSpawnActor(UObject* Asset, AActor* NewActor) override;
+	//~ End UActorFactory interface
+};
+
diff --git a/Engine/Plugins/Compositing/Composite/Source/CompositeEditor/Private/UI/SCompositePanelLayerTree.cpp b/Engine/Plugins/Compositing/Composite/Source/CompositeEditor/Private/UI/SCompositePanelLayerTree.cpp
--- a/Engine/Plugins/Compositing/Composite/Source/CompositeEditor/Private/UI/SCompositePanelLayerTree.cpp
+++ b/Engine/Plugins/Compositing/Composite/Source/CompositeEditor/Private/UI/SCompositePanelLayerTree.cpp
@@ -2,6 +2,7 @@
 
 #include "SCompositePanelLayerTree.h"
 
+#include "ActorFactories/ActorFactory.h"
 #include "CompositeActor.h"
 #include "CompositeEditorCommands.h"
 #include "Editor.h"
@@ -1886,7 +1887,23 @@ void SCompositePanelLayerTree::AddCompositeActor()
 
 	FScopedTransaction AddCompositeActorTransaction(LOCTEXT("AddCompositeActorTransaction", "Add Composite Actor"));
 	
-	ACompositeActor* AddedActor = Cast<ACompositeActor>(World->SpawnActor(ACompositeActor::StaticClass()));
+	ACompositeActor* AddedActor = nullptr;
+
+	if (GEditor)
+	{
+		UActorFactory* ActorFactory = GEditor->FindActorFactoryForActorClass(ACompositeActor::StaticClass());
+		if (ActorFactory)
+		{
+			AddedActor = Cast<ACompositeActor>(ActorFactory->CreateActor(ACompositeActor::StaticClass(), World->GetCurrentLevel(), FTransform()));
+		}
+	}
+
+	// Fallback to world spawn if the factory spawn failed
+	if (!AddedActor)
+	{
+		AddedActor = World->SpawnActor<ACompositeActor>(ACompositeActor::StaticClass());
+	}
+
 	if (AddedActor)
 	{
 		if (GEditor)
