Composite: Fixed undo breaking single active actor constraint.

* Conditionally deactivate previously active actor in the PostEditUndo callback.
* Additional: Actor factory now contains analytics tracking, and we avoid creating layers for the drag & drop actor.
* Additional: Fix incorrect MU session join deactivation, which would occur before a level load could replace actors. We now use the OnFinalizeWorkspaceSyncCompleted callback.

#rnx

diff --git a/Engine/Plugins/Compositing/Composite/Source/Composite/Private/CompositeActor.cpp b/Engine/Plugins/Compositing/Composite/Source/Composite/Private/CompositeActor.cpp
--- a/Engine/Plugins/Compositing/Composite/Source/Composite/Private/CompositeActor.cpp
+++ b/Engine/Plugins/Compositing/Composite/Source/Composite/Private/CompositeActor.cpp
@@ -549,6 +549,23 @@ void ACompositeActor::PostEditUndo()
 
 	PreEditCompositeLayers.Reset();
 
+	UWorld* World = GetWorld();
+
+	// Deactivate the current actor if any other actor is already active
+	// This can a occur when a previously active actor deletion is undone.
+	if (IsValid(World) && IsActive())
+	{
+		for (TActorIterator<ACompositeActor> It(World); It; ++It)
+		{
+			ACompositeActor* OtherActor = *It;
+			if ((OtherActor != this) && IsValid(OtherActor) && OtherActor->IsActive())
+			{
+				SetActive(false);
+				break;
+			}
+		}
+	}
+
 	Super::PostEditUndo();
 }
 
@@ -812,11 +829,11 @@ bool ACompositeActor::IsRendering() const
 	return bIsRendering;
 }
 
-void ACompositeActor::JoinConcertSession()
+void ACompositeActor::PostJoinConcertSession()
 {
 	SetActive(false);
 
-	OnJoinMultiUserSession.Broadcast();
+	OnPostJoinMultiUserSession.Broadcast();
 }
 
 void ACompositeActor::PropagateStateChange(bool bApply)
diff --git a/Engine/Plugins/Compositing/Composite/Source/Composite/Public/CompositeActor.h b/Engine/Plugins/Compositing/Composite/Source/Composite/Public/CompositeActor.h
--- a/Engine/Plugins/Compositing/Composite/Source/Composite/Public/CompositeActor.h
+++ b/Engine/Plugins/Compositing/Composite/Source/Composite/Public/CompositeActor.h
@@ -18,7 +18,7 @@ class UCompositeShadowReflectionCatcherComponent;
 class UPrimitiveComponent;
 class UCompositeSceneCapture2DComponent;
 
-DECLARE_DYNAMIC_MULTICAST_DELEGATE(FOnCompositeActorJoinMultiUserSession);
+DECLARE_DYNAMIC_MULTICAST_DELEGATE(FOnCompositeActorPostJoinMultiUserSession);
 
 USTRUCT()
 struct FSceneCaptureComponentArray
@@ -168,9 +168,9 @@ public:
 	UPROPERTY(EditAnywhere, BlueprintReadWrite, Category = "Composite", AdvancedDisplay)
 	ECompositeAllowedViewModes AllowedViewModes;
 
-	/** Delegate called when joining multi-user. Can be used to automatically re-active the actor locally. */
+	/** Delegate called after joining a multi-user session. Can be used to automatically re-active the actor locally. */
 	UPROPERTY(BlueprintAssignable, BlueprintCallable, BlueprintReadWrite, Category = "Composite")
-	FOnCompositeActorJoinMultiUserSession OnJoinMultiUserSession;
+	FOnCompositeActorPostJoinMultiUserSession OnPostJoinMultiUserSession;
 
 public:
 	/** Convenience function used by layers to request a scene capture component. */
@@ -183,8 +183,8 @@ public:
 	/** Returns true if the composite is actively rendering. */
 	UE_API bool IsRendering() const;
 
-	/** Called when we join a multi-user session. */
-	UE_API void JoinConcertSession();
+	/** Called after we join a multi-user session. */
+	UE_API void PostJoinConcertSession();
 
 private:
 	/**
diff --git a/Engine/Plugins/Compositing/Composite/Source/CompositeEditor/Private/ActorFactories/CompositeActorFactory.cpp b/Engine/Plugins/Compositing/Composite/Source/CompositeEditor/Private/ActorFactories/CompositeActorFactory.cpp
--- a/Engine/Plugins/Compositing/Composite/Source/CompositeEditor/Private/ActorFactories/CompositeActorFactory.cpp
+++ b/Engine/Plugins/Compositing/Composite/Source/CompositeEditor/Private/ActorFactories/CompositeActorFactory.cpp
@@ -3,6 +3,7 @@
 #include "CompositeActorFactory.h"
 
 #include "CompositeActor.h"
+#include "EngineAnalytics.h"
 #include "Layers/CompositeLayerMainRender.h"
 #include "Layers/CompositeLayerShadowReflection.h"
 #include "Layers/CompositeLayerPlate.h"
@@ -24,8 +25,8 @@ void UCompositeActorFactory::PostSpawnActor(UObject* Asset, AActor* NewActor)
 
 	ACompositeActor* CompositeActor = Cast<ACompositeActor>(NewActor);
 
-	// Configure newly spawned composite actor instance
-	if (IsValid(CompositeActor) && !CompositeActor->IsTemplate())
+	// Configure newly spawned composite actor instance (except the transient drag & drop actor).
+	if (IsValid(CompositeActor) && !CompositeActor->IsTemplate() && !CompositeActor->HasAnyFlags(RF_Transient))
 	{
 		// Default layers when creating a new actor
 		TArray<TObjectPtr<UCompositeLayerBase>> NewLayers;
@@ -43,12 +44,19 @@ void UCompositeActorFactory::PostSpawnActor(UObject* Asset, AActor* NewActor)
 		// By default, we disable the shadow reflection catcher due to its significant added cost.
 		NewLayers[1]->SetEnabled(false);
 
+		// No need to call Modify() on the spawned actor
 		CompositeActor->SetCompositeLayers(NewLayers);
 
 		// Refresh Composure panel
 		FProperty* CompositeLayersProperty = FindFProperty<FProperty>(ACompositeActor::StaticClass(), GET_MEMBER_NAME_CHECKED(ACompositeActor, CompositeLayers));
 		FPropertyChangedEvent LayersChangedEvent(CompositeLayersProperty, EPropertyChangeType::ArrayAdd);
 		FCoreUObjectDelegates::OnObjectPropertyChanged.Broadcast(NewActor, LayersChangedEvent);
+
+		// Record analytics
+		if (FEngineAnalytics::IsAvailable())
+		{
+			FEngineAnalytics::GetProvider().RecordEvent(TEXT("Editor.Usage.Composite.PostSpawnActor"));
+		}
 	}
 }
 
diff --git a/Engine/Plugins/Compositing/Composite/Source/CompositeEditor/Private/CompositeEditorModule.cpp b/Engine/Plugins/Compositing/Composite/Source/CompositeEditor/Private/CompositeEditorModule.cpp
--- a/Engine/Plugins/Compositing/Composite/Source/CompositeEditor/Private/CompositeEditorModule.cpp
+++ b/Engine/Plugins/Compositing/Composite/Source/CompositeEditor/Private/CompositeEditorModule.cpp
@@ -4,9 +4,9 @@
 
 #include "Editor.h"
 #include "EngineUtils.h"
-#include "EngineAnalytics.h"
 #include "HoldoutCompositeComponent.h"
 #include "Framework/Notifications/NotificationManager.h"
+#include "IConcertClientWorkspace.h"
 #include "IConcertSyncClient.h"
 #include "IConcertSyncClientModule.h"
 #include "LevelEditor.h"
@@ -49,7 +49,14 @@ void FCompositeEditorModule::ShutdownModule()
 {
 	if (ConcertSyncClient.IsValid())
 	{
-		ConcertSyncClient->GetConcertClient()->OnSessionStartup().RemoveAll(this);
+		ConcertSyncClient->OnWorkspaceStartup().RemoveAll(this);
+		ConcertSyncClient->OnWorkspaceShutdown().RemoveAll(this);
+		
+		TSharedPtr<IConcertClientWorkspace> WorkspacePtr = Workspace.Pin();
+		if (WorkspacePtr.IsValid())
+		{
+			WorkspacePtr->OnFinalizeWorkspaceSyncCompleted().RemoveAll(this);
+		}
 	}
 
 	UnregisterCustomizations();
@@ -78,7 +85,8 @@ void FCompositeEditorModule::OnPostEngineInit()
 	ConcertSyncClient = IConcertSyncClientModule::Get().GetClient(TEXT("MultiUser"));
 	if (ConcertSyncClient.IsValid())
 	{
-		ConcertSyncClient->GetConcertClient()->OnSessionStartup().AddRaw(this, &FCompositeEditorModule::OnConcertSessionStartup);
+		ConcertSyncClient->OnWorkspaceStartup().AddRaw(this, &FCompositeEditorModule::HandleWorkspaceStartup);
+		ConcertSyncClient->OnWorkspaceShutdown().AddRaw(this, &FCompositeEditorModule::HandleWorkspaceShutdown);
 	}
 }
 
@@ -112,14 +120,33 @@ void FCompositeEditorModule::OnLevelActorAdded(AActor* InActor)
 			break;
 		}
 	}
+}
 
-	if (FEngineAnalytics::IsAvailable())
+void FCompositeEditorModule::HandleWorkspaceStartup(const TSharedPtr<IConcertClientWorkspace>& NewWorkspace)
+{
+	Workspace = NewWorkspace;
+
+	if (NewWorkspace.IsValid())
 	{
-		FEngineAnalytics::GetProvider().RecordEvent(TEXT("Editor.Usage.Composite.AddCompositeActorToLevel"));
+		NewWorkspace->OnFinalizeWorkspaceSyncCompleted().AddRaw(this, &FCompositeEditorModule::HandleFinalizeWorkspaceSyncCompleted);
+	}
+}
+
+void FCompositeEditorModule::HandleWorkspaceShutdown(const TSharedPtr<IConcertClientWorkspace>& WorkspaceShuttingDown)
+{
+	if (WorkspaceShuttingDown == Workspace)
+	{
+		TSharedPtr<IConcertClientWorkspace> WorkspacePtr = Workspace.Pin();
+		if (WorkspacePtr.IsValid())
+		{
+			WorkspacePtr->OnFinalizeWorkspaceSyncCompleted().RemoveAll(this);
+		}
+
+		Workspace.Reset();
 	}
 }
 
-void FCompositeEditorModule::OnConcertSessionStartup(TSharedRef<IConcertClientSession> InSession)
+void FCompositeEditorModule::HandleFinalizeWorkspaceSyncCompleted()
 {
 	UWorld* World = nullptr;
 	if (GEditor)
@@ -134,7 +161,7 @@ void FCompositeEditorModule::OnConcertSessionStartup(TSharedRef<IConcertClientSe
 			ACompositeActor* CompositeActor = *It;
 			if (IsValid(CompositeActor))
 			{
-				CompositeActor->JoinConcertSession();
+				CompositeActor->PostJoinConcertSession();
 			}
 		}
 	}
diff --git a/Engine/Plugins/Compositing/Composite/Source/CompositeEditor/Private/UI/SCompositePanelLayerTree.cpp b/Engine/Plugins/Compositing/Composite/Source/CompositeEditor/Private/UI/SCompositePanelLayerTree.cpp
--- a/Engine/Plugins/Compositing/Composite/Source/CompositeEditor/Private/UI/SCompositePanelLayerTree.cpp
+++ b/Engine/Plugins/Compositing/Composite/Source/CompositeEditor/Private/UI/SCompositePanelLayerTree.cpp
@@ -510,9 +510,10 @@ private:
 	{
 		if (TStrongObjectPtr<ACompositeActor> CompositeActor = GetCompositeActor())
 		{
+			const bool bIsActive = (InCheckBoxState == ECheckBoxState::Checked);
 			// No transactions, the property is non-transactional
-			CompositeActor->SetActive(InCheckBoxState == ECheckBoxState::Checked);
-			OnCompositeActorActivated.ExecuteIfBound(CompositeActor);
+			CompositeActor->SetActive(bIsActive);
+			OnCompositeActorActivated.ExecuteIfBound(bIsActive ? CompositeActor : nullptr);
 		}
 	}
 
@@ -1978,6 +1979,9 @@ void SCompositePanelLayerTree::OnCompositeActorActivated(const TStrongObjectPtr<
 		return;
 	}
 
+	// Keep last active actor up-to-date
+	LastActiveActor = CompositeActor.Get();
+
 	AActor* CompositeCameraActor = CompositeActor->GetCamera().OtherActor.Get();
 
 	if (!IsValid(CompositeCameraActor))
diff --git a/Engine/Plugins/Compositing/Composite/Source/CompositeEditor/Public/CompositeEditorModule.h b/Engine/Plugins/Compositing/Composite/Source/CompositeEditor/Public/CompositeEditorModule.h
--- a/Engine/Plugins/Compositing/Composite/Source/CompositeEditor/Public/CompositeEditorModule.h
+++ b/Engine/Plugins/Compositing/Composite/Source/CompositeEditor/Public/CompositeEditorModule.h
@@ -7,7 +7,7 @@
 #include "Logging/LogMacros.h"
 #include "Modules/ModuleManager.h"
 
-class IConcertClientSession;
+class IConcertClientWorkspace;
 
 /** Log category for this module. */
 DECLARE_LOG_CATEGORY_EXTERN(LogCompositeEditor, Log, All);
@@ -27,9 +27,6 @@ private:
 	/** Called when an Actor has been added to a level. */
 	void OnLevelActorAdded(AActor* InActor);
 
-	/** Callback on multi-user concert session start. */
-	void OnConcertSessionStartup(TSharedRef<IConcertClientSession> InSession);
-
 	/** Callback to trigger warnings when an holdout composite component is created with a composite actor. */
 	void TriggerHoldoutCompositeWarning(const class UHoldoutCompositeComponent* InComponent);
 
@@ -42,10 +39,22 @@ private:
 	void RegisterCustomizations();
 	void UnregisterCustomizations();
 
+	/** Workspace startup callback. */
+	void HandleWorkspaceStartup(const TSharedPtr<IConcertClientWorkspace>& NewWorkspace);
+	
+	/** Workspace shutdown callback. */
+	void HandleWorkspaceShutdown(const TSharedPtr<IConcertClientWorkspace>& WorkspaceShuttingDown);
+
+	/** Workspace sync completed callback. */
+	void HandleFinalizeWorkspaceSyncCompleted();
+
 private:
 	/** Cached concert sync client. */
 	TSharedPtr<class IConcertSyncClient> ConcertSyncClient;
 
 	/** Tab manager change delegate. */
 	FDelegateHandle OnTabManagerChangedDelegateHandle;
+
+	/** Holds a weak pointer to the current workspace. */
+	TWeakPtr<IConcertClientWorkspace> Workspace;
 };
