Composite: Fixed pre-processing passes incorrectly running twice.

* Duplicate work was caused by branching SSR logic when registering pass proxies every frame.
* Prevented external textures from being registered more than once via a new convenience function.
* Simplified calls to register external textures: the convenience function now simply returns the ResourceId.

#rnx

diff --git a/Engine/Plugins/Compositing/Composite/Source/Composite/Private/CompositeActor.cpp b/Engine/Plugins/Compositing/Composite/Source/Composite/Private/CompositeActor.cpp
--- a/Engine/Plugins/Compositing/Composite/Source/Composite/Private/CompositeActor.cpp
+++ b/Engine/Plugins/Compositing/Composite/Source/Composite/Private/CompositeActor.cpp
@@ -172,6 +172,12 @@ namespace UE
 					return {};
 				}
 
+				/**
+				* Register pre-processing passes before SSR branch below to prevent duplicate work.
+				* These passes only run once at the start of a frame (see ApplyPreprocessing call site in FCompositeCoreSceneViewExtension).
+				*/
+				RenderWork.PreprocessingPasses = Context.PreprocessingPasses;
+
 				if (InCompositeActor.bEnableScreenSpaceReflections)
 				{
 					const int32 BranchSourceIndex = BranchLayerIndex - 1;
@@ -202,8 +208,7 @@ namespace UE
 					RenderWork.FramePasses.Emplace(ISceneViewExtension::EPostProcessingPass::Tonemap, LayerProxies);
 				}
 
-				RenderWork.ExternalInputs = Context.ExternalTextures;
-				RenderWork.PreprocessingPasses = Context.PreprocessingPasses;
+				RenderWork.ExternalInputs = Context.GetExternalTextures();
 
 				switch (InCompositeActor.MainRenderOutput)
 				{
diff --git a/Engine/Plugins/Compositing/Composite/Source/Composite/Private/Layers/CompositeLayerBase.cpp b/Engine/Plugins/Compositing/Composite/Source/Composite/Private/Layers/CompositeLayerBase.cpp
--- a/Engine/Plugins/Compositing/Composite/Source/Composite/Private/Layers/CompositeLayerBase.cpp
+++ b/Engine/Plugins/Compositing/Composite/Source/Composite/Private/Layers/CompositeLayerBase.cpp
@@ -140,3 +140,26 @@ UE::CompositeCore::FPassInputDecl UCompositeLayerBase::GetDefaultSecondInput(con
 	return InputDecl;
 }
 
+UE::CompositeCore::ResourceId UCompositeLayerBase::FTraversalContext::FindOrCreateExternalTexture(TWeakObjectPtr<UTexture> InTexture, UE::CompositeCore::FResourceMetadata InMetadata)
+{
+	int32 TextureIndex = ExternalTextures.IndexOfByPredicate([&InTexture, &InMetadata](const UE::CompositeCore::FExternalTexture& InExternalTexture)
+		{
+			return (InExternalTexture.Texture.Get() == InTexture.Get()) && (InExternalTexture.Metadata == InMetadata);
+		}
+	);
+
+	if (TextureIndex == INDEX_NONE)
+	{
+		ExternalTextures.Add(UE::CompositeCore::FExternalTexture{ MoveTemp(InTexture), MoveTemp(InMetadata) });
+		
+		TextureIndex = ExternalTextures.Num() - 1;
+	}
+
+	return UE::CompositeCore::EXTERNAL_RANGE_START_ID + TextureIndex;
+}
+
+const TArray<UE::CompositeCore::FExternalTexture>& UCompositeLayerBase::FTraversalContext::GetExternalTextures() const
+{
+	return ExternalTextures;
+}
+
diff --git a/Engine/Plugins/Compositing/Composite/Source/Composite/Private/Layers/CompositeLayerPlate.cpp b/Engine/Plugins/Compositing/Composite/Source/Composite/Private/Layers/CompositeLayerPlate.cpp
--- a/Engine/Plugins/Compositing/Composite/Source/Composite/Private/Layers/CompositeLayerPlate.cpp
+++ b/Engine/Plugins/Compositing/Composite/Source/Composite/Private/Layers/CompositeLayerPlate.cpp
@@ -457,11 +457,9 @@ UE::CompositeCore::ResourceId UCompositeLayerPlate::AddPreprocessingPasses(
 		{
 			if (PassIndex == LastValidPassIndex)
 			{
-				FExternalTexture& ExternalTexture = InContext.ExternalTextures.AddDefaulted_GetRef();
-				ExternalTexture.Texture = GetRenderTargetFn();
-				ExternalTexture.Metadata.bDistorted = true;
-
-				OutputTextureId = EXTERNAL_RANGE_START_ID + static_cast<ResourceId>(InContext.ExternalTextures.Num() - 1);
+				FResourceMetadata Metadata = {};
+				Metadata.bDistorted = true;
+				OutputTextureId = InContext.FindOrCreateExternalTexture(GetRenderTargetFn(), Metadata);
 
 				PrepassProxy->DeclareOutputOverride(OutputTextureId);
 			}
@@ -482,10 +480,9 @@ bool UCompositeLayerPlate::GetProxy(FTraversalContext& InContext, FSceneRenderin
 	if (IsValid(Texture))
 	{
 		// Register source texture into the frame work context (so that post-processing passes can refer to it).
-		FExternalTexture& ExternalTexture = InContext.ExternalTextures.AddDefaulted_GetRef();
-		ExternalTexture.Texture = Texture;
-		ExternalTexture.Metadata.bDistorted = true;
-		TextureId = EXTERNAL_RANGE_START_ID + static_cast<ResourceId>(InContext.ExternalTextures.Num() - 1);
+		FResourceMetadata Metadata = {};
+		Metadata.bDistorted = true;
+		TextureId = InContext.FindOrCreateExternalTexture(Texture, Metadata);
 		
 		const ResourceId OriginalTextureId = TextureId;
 
diff --git a/Engine/Plugins/Compositing/Composite/Source/Composite/Private/Layers/CompositeLayerSceneCapture.cpp b/Engine/Plugins/Compositing/Composite/Source/Composite/Private/Layers/CompositeLayerSceneCapture.cpp
--- a/Engine/Plugins/Compositing/Composite/Source/Composite/Private/Layers/CompositeLayerSceneCapture.cpp
+++ b/Engine/Plugins/Compositing/Composite/Source/Composite/Private/Layers/CompositeLayerSceneCapture.cpp
@@ -92,15 +92,14 @@ bool UCompositeLayerSceneCapture::GetProxy(FTraversalContext& InContext, FSceneR
 
 	FCompositeRenderTargetPool::Get().ConditionalAcquireTarget(CaptureComponent, CaptureComponent->TextureTarget, GetRenderResolution());
 
-	FExternalTexture& ExternalTexture = InContext.ExternalTextures.AddDefaulted_GetRef();
-	ExternalTexture.Texture = CaptureComponent->TextureTarget;
-	ExternalTexture.Metadata.bInvertedAlpha = true;
-	ExternalTexture.Metadata.bDistorted = CaptureComponent->ShowFlags.LensDistortion;
-	ExternalTexture.Metadata.Encoding = UE::CompositeCore::EEncoding::Linear;
+	FResourceMetadata Metadata = {};
+	Metadata.bInvertedAlpha = true;
+	Metadata.bDistorted = CaptureComponent->ShowFlags.LensDistortion;
+
+	const ResourceId TexId = InContext.FindOrCreateExternalTexture(CaptureComponent->TextureTarget, Metadata);
 
-	const int32 TexIndex = InContext.ExternalTextures.Num() - 1;
 	FPassInputDecl PassInput;
-	PassInput.Set<FPassExternalResourceDesc>({ EXTERNAL_RANGE_START_ID + TexIndex });
+	PassInput.Set<FPassExternalResourceDesc>({ TexId });
 
 	AddChildPasses(PassInput, InContext, InFrameAllocator, LayerPasses);
 
diff --git a/Engine/Plugins/Compositing/Composite/Source/Composite/Private/Layers/CompositeLayerShadowReflection.cpp b/Engine/Plugins/Compositing/Composite/Source/Composite/Private/Layers/CompositeLayerShadowReflection.cpp
--- a/Engine/Plugins/Compositing/Composite/Source/Composite/Private/Layers/CompositeLayerShadowReflection.cpp
+++ b/Engine/Plugins/Compositing/Composite/Source/Composite/Private/Layers/CompositeLayerShadowReflection.cpp
@@ -95,14 +95,12 @@ bool UCompositeLayerShadowReflection::GetProxy(FTraversalContext& InContext, FSc
 			UCompositeShadowReflectionCatcherComponent* Component = CachedSceneCaptures[Index].Get();
 			FCompositeRenderTargetPool::Get().ConditionalAcquireTarget(Component, Component->TextureTarget, GetRenderResolution());
 
-			FExternalTexture& ExternalTexture = InContext.ExternalTextures.AddDefaulted_GetRef();
-			ExternalTexture.Texture = Component->TextureTarget;
-			ExternalTexture.Metadata.bInvertedAlpha = true;
-			ExternalTexture.Metadata.bDistorted = Component->ShowFlags.LensDistortion;
-			ExternalTexture.Metadata.Encoding = UE::CompositeCore::EEncoding::Linear;
-
-			int32 TexIndex = InContext.ExternalTextures.Num() - 1;
-			Inputs[Index].Set<FPassExternalResourceDesc>({ EXTERNAL_RANGE_START_ID + TexIndex });
+			FResourceMetadata Metadata = {};
+			Metadata.bInvertedAlpha = true;
+			Metadata.bDistorted = Component->ShowFlags.LensDistortion;
+			const ResourceId TexId = InContext.FindOrCreateExternalTexture(Component->TextureTarget, Metadata);
+
+			Inputs[Index].Set<FPassExternalResourceDesc>({ TexId });
 		}
 
 		DivisionProxy = InFrameAllocator.Create<FMergePassProxy>(MoveTemp(Inputs), ECompositeCoreMergeOp::Divide, TEXT("ShadowRefl Div"), ELensDistortionHandling::Disabled);
diff --git a/Engine/Plugins/Compositing/Composite/Source/Composite/Private/Layers/CompositeLayerSingleLightShadow.cpp b/Engine/Plugins/Compositing/Composite/Source/Composite/Private/Layers/CompositeLayerSingleLightShadow.cpp
--- a/Engine/Plugins/Compositing/Composite/Source/Composite/Private/Layers/CompositeLayerSingleLightShadow.cpp
+++ b/Engine/Plugins/Compositing/Composite/Source/Composite/Private/Layers/CompositeLayerSingleLightShadow.cpp
@@ -249,18 +249,16 @@ bool UCompositeLayerSingleLightShadow::GetProxy(FTraversalContext& InContext, FS
 	FCompositeSingleLightShadowProxy* ShadowProxy;
 	{
 		// Display-resolution scene depth without jitter as an external texture.
-		InContext.ExternalTextures.Add(FExternalTexture{ CachedSceneDepthCapture->TextureTarget, FResourceMetadata{} });
-		const int32 SceneDepthIndex = InContext.ExternalTextures.Num() - 1;
+		const ResourceId SceneDepthId = InContext.FindOrCreateExternalTexture(CachedSceneDepthCapture->TextureTarget, {});
 
 		// Shadow map texture
-		InContext.ExternalTextures.Add(FExternalTexture{ CachedShadowMapCapture->TextureTarget, FResourceMetadata{} });
-		const int32 ShadowMapIndex = InContext.ExternalTextures.Num() - 1;
+		const ResourceId ShadowMapId = InContext.FindOrCreateExternalTexture(CachedShadowMapCapture->TextureTarget, {});
 
 		// Setup shadow pass inputs
 		FPassInputDeclArray PassInputs;
 		PassInputs.SetNum(2);
-		PassInputs[0].Set<FPassExternalResourceDesc>({ EXTERNAL_RANGE_START_ID + SceneDepthIndex });
-		PassInputs[1].Set<FPassExternalResourceDesc>({ EXTERNAL_RANGE_START_ID + ShadowMapIndex });
+		PassInputs[0].Set<FPassExternalResourceDesc>({ SceneDepthId });
+		PassInputs[1].Set<FPassExternalResourceDesc>({ ShadowMapId });
 
 		static const TConsoleVariableData<float>* CVarDepthBias = IConsoleManager::Get().FindTConsoleVariableDataFloat(TEXT("r.Shadow.CSMDepthBias"));
 		float ShadowDepthBias = CVarDepthBias ? CVarDepthBias->GetValueOnGameThread() : 10.0f;
diff --git a/Engine/Plugins/Compositing/Composite/Source/Composite/Public/Layers/CompositeLayerBase.h b/Engine/Plugins/Compositing/Composite/Source/Composite/Public/Layers/CompositeLayerBase.h
--- a/Engine/Plugins/Compositing/Composite/Source/Composite/Public/Layers/CompositeLayerBase.h
+++ b/Engine/Plugins/Compositing/Composite/Source/Composite/Public/Layers/CompositeLayerBase.h
@@ -35,11 +35,22 @@ public:
 		/** Flag to indicate that a layer is independent of scene textures. Work from scene-independent layers can be shared between post-processing locations. */
 		bool bNeedsSceneTextures = false;
 		
-		/** List of external textures */
-		TArray<UE::CompositeCore::FExternalTexture> ExternalTextures;
-		
 		/** Array of passes applied at the start of rendering. */
 		TSortedMap<UE::CompositeCore::ResourceId, TArray<const FCompositeCorePassProxy*>> PreprocessingPasses;
+
+		/**
+		* Find or create an external texture resource.
+		* 
+		* @return ResourceId Texture resource identifier
+		*/
+		UE::CompositeCore::ResourceId FindOrCreateExternalTexture(TWeakObjectPtr<UTexture> InTexture, UE::CompositeCore::FResourceMetadata InMetadata);
+
+		/** Get external textures. */
+		const TArray<UE::CompositeCore::FExternalTexture>& GetExternalTextures() const;
+
+	private:
+		/** List of external textures */
+		TArray<UE::CompositeCore::FExternalTexture> ExternalTextures;
 	};
 
 	/** Constructor. */
diff --git a/Engine/Plugins/Compositing/CompositeCore/Source/CompositeCore/Public/Passes/CompositeCorePassProxy.h b/Engine/Plugins/Compositing/CompositeCore/Source/CompositeCore/Public/Passes/CompositeCorePassProxy.h
--- a/Engine/Plugins/Compositing/CompositeCore/Source/CompositeCore/Public/Passes/CompositeCorePassProxy.h
+++ b/Engine/Plugins/Compositing/CompositeCore/Source/CompositeCore/Public/Passes/CompositeCorePassProxy.h
@@ -62,6 +62,12 @@ namespace UE
 			
 			/** Source color encoding. */
 			EEncoding Encoding = EEncoding::Linear;
+			
+			/** Equality operator */
+			bool operator==(const FResourceMetadata& Other) const = default;
+			
+			/** Inequality operator */
+			bool operator!=(const FResourceMetadata& Other) const = default;
 		};
 
 		/** Pass texture description for internal resources (default scene textures). */
