Composite: Fixed warning when creating a new composite mesh actor.

* Additional fix to default composite mesh scale and positioning so that its corner is camera-aligned, and generally covers more of the camera FoV.
* Asset modification repositions the mesh pivot to the bottom center.

#rnx

diff --git a/Engine/Plugins/Compositing/Composite/Source/Composite/Private/CompositeMeshActor.cpp b/Engine/Plugins/Compositing/Composite/Source/Composite/Private/CompositeMeshActor.cpp
--- a/Engine/Plugins/Compositing/Composite/Source/Composite/Private/CompositeMeshActor.cpp
+++ b/Engine/Plugins/Compositing/Composite/Source/Composite/Private/CompositeMeshActor.cpp
@@ -8,7 +8,7 @@ ACompositeMeshActor::ACompositeMeshActor(const FObjectInitializer& ObjectInitial
 	: Super(ObjectInitializer)
 {
 	CompositeMeshComponent = CreateDefaultSubobject<UCompositeMeshComponent>(TEXT("DefaultCompositeMeshComponent"));
-	CompositeMeshComponent->SetupAttachment(RootComponent);
+	SetRootComponent(CompositeMeshComponent);
 }
 
 ACompositeMeshActor::~ACompositeMeshActor() = default;
diff --git a/Engine/Plugins/Compositing/Composite/Source/CompositeEditor/Private/Customizations/CompositeLayerPlateCustomization.cpp b/Engine/Plugins/Compositing/Composite/Source/CompositeEditor/Private/Customizations/CompositeLayerPlateCustomization.cpp
--- a/Engine/Plugins/Compositing/Composite/Source/CompositeEditor/Private/Customizations/CompositeLayerPlateCustomization.cpp
+++ b/Engine/Plugins/Compositing/Composite/Source/CompositeEditor/Private/Customizations/CompositeLayerPlateCustomization.cpp
@@ -323,9 +323,8 @@ void FCompositeLayerPlateCustomization::CreateCompositeMeshActor()
 		
 		const FVector CameraLoc = CameraComponent->GetComponentLocation();
 
-		// Mesh location is set so that the center of the mesh's "front" wall aligns with the camera, and is
-		// a distance from it such that the most of the full wall is visible from the camera
-		const FVector MeshLoc = CameraLoc - 0.5 * MeshSize * Up - 0.5 * MeshSize * Right + (Distance - MeshSize) * Forward;
+		// Mesh location is set it is horizontally centered with the camera
+		const FVector MeshLoc = CameraLoc - 0.333 * MeshSize * Up + (Distance - MeshSize) * Forward;
 		
 		SpawnTransform.SetLocation(MeshLoc);
 		
@@ -333,7 +332,11 @@ void FCompositeLayerPlateCustomization::CreateCompositeMeshActor()
 		// Only preserve the yaw to keep the mesh horizontally level
 		Rotation.Pitch = 0.0f;
 		Rotation.Roll = 0.0f;
+		Rotation.Yaw -= 45.0f; // Additional rotation so that the mesh corner is aligned with the camera direction
 		SpawnTransform.SetRotation(Rotation.Quaternion());
+
+		// Increase default mesh scale
+		SpawnTransform.SetScale3D(2.0 * FVector::One());
 	}
 
 	FScopedTransaction AddCompositeMeshTransaction(LOCTEXT("AddCompositeMeshTransaction", "Add Composite Mesh"));
@@ -385,7 +388,7 @@ bool FCompositeLayerPlateCustomization::CanCreateCompositeMeshActor() const
 		return false;
 	}
 
-	ACompositeActor* CompositeActor = Objects[0]->GetTypedOuter<ACompositeActor>();
+	ACompositeActor* CompositeActor = Objects[0].IsValid() ? Objects[0]->GetTypedOuter<ACompositeActor>() : nullptr;
 	if (!CompositeActor)
 	{
 		return false;
