Composure: Prevent "leaked" scene capture component after layer deletion.

* Previous deletion logic incorrectly relied on the owner still being valid, but this was fixed to avoid unnecessary ticking at CL 46604443 by renaming the deleted objects.
* Isolated to 5.7 only, since this is already fixed in Main at CL 47600697.

#rnx

diff --git a/Engine/Plugins/Compositing/Composite/Source/Composite/Private/CompositeActor.cpp b/Engine/Plugins/Compositing/Composite/Source/Composite/Private/CompositeActor.cpp
--- a/Engine/Plugins/Compositing/Composite/Source/Composite/Private/CompositeActor.cpp
+++ b/Engine/Plugins/Compositing/Composite/Source/Composite/Private/CompositeActor.cpp
@@ -501,6 +501,9 @@ void ACompositeActor::PostEditChangeChainProperty(struct FPropertyChangedChainEv
 				{
 					if (!CompositeLayers.Contains(PreEditCompositeLayer))
 					{
+						// 5.7 hotfix: Ensure scene capture component destruction
+						DestroySceneCaptures(PreEditCompositeLayer);
+
 						PreEditCompositeLayer->OnRemoved(GetWorld());
 					}
 				}
@@ -547,6 +550,9 @@ void ACompositeActor::PostEditUndo()
 		{
 			if (!CompositeLayers.Contains(PreEditCompositeLayer))
 			{
+				// 5.7 hotfix: Ensure scene capture component destruction
+				DestroySceneCaptures(PreEditCompositeLayer);
+
 				PreEditCompositeLayer->OnRemoved(GetWorld());
 			}
 		}
