Composite: Aggregation of fixes already submitted to 5.8.

* UE-354350: Fixed validation failure when running with DX11.
* UE-352161: Fixed processing layers not running unless other layers are present.
* UE-351167: Fixed disappearing blueprint meshes in scene captures layers, which now also default to TAA.
* UE-354526: Fixed inverted color grading behavior.
* UE-335900: Fixed single light shadow issues: mesh registration now consistent with shadow/reflection layer, and self-shadowing/biasing is no longer needed.
* Removed unwanted async compute for denoise pass.

diff --git a/Engine/Plugins/Compositing/Composite/Source/Composite/Private/Components/CompositeSceneCapture2DComponent.cpp b/Engine/Plugins/Compositing/Composite/Source/Composite/Private/Components/CompositeSceneCapture2DComponent.cpp
--- a/Engine/Plugins/Compositing/Composite/Source/Composite/Private/Components/CompositeSceneCapture2DComponent.cpp
+++ b/Engine/Plugins/Compositing/Composite/Source/Composite/Private/Components/CompositeSceneCapture2DComponent.cpp
@@ -49,7 +49,11 @@ UCompositeSceneCapture2DComponent::UCompositeSceneCapture2DComponent(const FObje
 	ShowFlags.SetEyeAdaptation(false);
 	ShowFlags.SetLocalExposure(false);
 	ShowFlags.SetPostProcessMaterial(false);
-
+	/**
+	 * This can cause temporal ghosting artifacts on moving shadows, but remains preferable over the otherwise default shadow noise.
+	 * Automatically ignored with custom render passes, where post-processing does not apply.
+	 */
+	ShowFlags.SetTemporalAA(true);
 }
 
 void UCompositeSceneCapture2DComponent::BeginDestroy()
diff --git a/Engine/Plugins/Compositing/Composite/Source/Composite/Private/Components/CompositeShadowReflectionCatcherComponent.cpp b/Engine/Plugins/Compositing/Composite/Source/Composite/Private/Components/CompositeShadowReflectionCatcherComponent.cpp
--- a/Engine/Plugins/Compositing/Composite/Source/Composite/Private/Components/CompositeShadowReflectionCatcherComponent.cpp
+++ b/Engine/Plugins/Compositing/Composite/Source/Composite/Private/Components/CompositeShadowReflectionCatcherComponent.cpp
@@ -15,8 +15,6 @@ UCompositeShadowReflectionCatcherComponent::UCompositeShadowReflectionCatcherCom
 
 	ShowFlags.SetGrain(false);
 	ShowFlags.SetScreenSpaceReflections(false);
-	//Note: This can cause temporal ghosting artifacts on moving shadows, but remains preferable over the otherwise default shadow noise.
-	ShowFlags.SetTemporalAA(true);
 }
 
 #undef LOCTEXT_NAMESPACE
diff --git a/Engine/Plugins/Compositing/Composite/Source/Composite/Private/Layers/CompositeLayerSceneCapture.cpp b/Engine/Plugins/Compositing/Composite/Source/Composite/Private/Layers/CompositeLayerSceneCapture.cpp
--- a/Engine/Plugins/Compositing/Composite/Source/Composite/Private/Layers/CompositeLayerSceneCapture.cpp
+++ b/Engine/Plugins/Compositing/Composite/Source/Composite/Private/Layers/CompositeLayerSceneCapture.cpp
@@ -15,9 +15,20 @@ UCompositeLayerSceneCapture::UCompositeLayerSceneCapture(const FObjectInitialize
 	, bCustomRenderPass{false}
 {
 	Operation = ECompositeCoreMergeOp::Over;
+
+	if (!HasAnyFlags(RF_ClassDefaultObject))
+	{
+		// 5.7 hotfix: We must use delegates since adding virtual function overrides is not permitted
+		FWorldDelegates::OnWorldPreSendAllEndOfFrameUpdates.AddUObject(this, &UCompositeLayerSceneCapture::OnEndOfFrameUpdate);
+#if WITH_EDITOR
+		FCoreUObjectDelegates::OnPreObjectPropertyChanged.AddUObject(this, &UCompositeLayerSceneCapture::OnObjectPreEditChange);
+#endif
+	}
 }
 
+PRAGMA_DISABLE_DEPRECATION_WARNINGS
 UCompositeLayerSceneCapture::~UCompositeLayerSceneCapture() = default;
+PRAGMA_ENABLE_DEPRECATION_WARNINGS
 
 void UCompositeLayerSceneCapture::OnRemoved(const UWorld* World)
 {
@@ -36,6 +47,11 @@ void UCompositeLayerSceneCapture::BeginDestroy()
 	Super::BeginDestroy();
 
 	OnRemoved(GetWorld()); // Redundant remove call for safety
+
+	FWorldDelegates::OnWorldPreSendAllEndOfFrameUpdates.RemoveAll(this);
+#if WITH_EDITOR
+	FCoreUObjectDelegates::OnPreObjectPropertyChanged.RemoveAll(this);
+#endif
 }
 
 #if WITH_EDITOR
@@ -62,6 +78,8 @@ void UCompositeLayerSceneCapture::PostEditChangeProperty(FPropertyChangedEvent&
 void UCompositeLayerSceneCapture::PreEditUndo()
 {
 	Super::PreEditUndo();
+
+	RestorePrimitiveVisibilityState();
 }
 
 void UCompositeLayerSceneCapture::PostEditUndo()
@@ -119,38 +137,20 @@ const TArray<TSoftObjectPtr<AActor>> UCompositeLayerSceneCapture::GetActors() co
 
 void UCompositeLayerSceneCapture::SetActors(TArray<TSoftObjectPtr<AActor>> InActors)
 {
-	TArray<TWeakObjectPtr<UPrimitiveComponent>> ShowOnlyComponents;
-	ShowOnlyComponents.Reserve(InActors.Num());
+	RestorePrimitiveVisibilityState();
 
-	for (const TSoftObjectPtr<AActor>& SoftActor : InActors)
-	{
-		if (const AActor* Actor = SoftActor.Get())
-		{
-			for (UActorComponent* Component : Actor->GetComponents())
-			{
-				if (UPrimitiveComponent* PrimitiveComponent = Cast<UPrimitiveComponent>(Component))
-				{
-					ShowOnlyComponents.Add(PrimitiveComponent);
-				}
-			}
-		}
-	}
+	Actors = MoveTemp(InActors);
+
+	UpdatePrimitiveVisibilityState();
 
 	USceneCaptureComponent2D* SceneCaptureComponent = GetSceneCaptureComponent();
 	if (IsValid(SceneCaptureComponent))
 	{
 		// First, we empty show-only actors since our previous logic used them
 		SceneCaptureComponent->ShowOnlyActors.Empty();
-		SceneCaptureComponent->ShowOnlyComponents = ShowOnlyComponents;
+		
+		UpdateSceneCaptureComponents(*SceneCaptureComponent);
 	}
-
-	RestorePrimitiveVisibilityState();
-
-	Actors = MoveTemp(InActors);
-
-	CachePrimitiveVisibilityState(ShowOnlyComponents);
-
-	UpdatePrimitiveVisibilityState();
 }
 
 bool UCompositeLayerSceneCapture::IsCustomRenderPass() const
@@ -201,22 +201,29 @@ UCompositeSceneCapture2DComponent* UCompositeLayerSceneCapture::GetSceneCaptureC
 	return nullptr;
 }
 
-void UCompositeLayerSceneCapture::CachePrimitiveVisibilityState(TArrayView<TWeakObjectPtr<UPrimitiveComponent>> InPrimitives)
+void UCompositeLayerSceneCapture::UpdatePrimitiveVisibilityState()
 {
-	CachedVisibilityInSceneCapture.Reset();
-
-	for (const TWeakObjectPtr<UPrimitiveComponent>& PrimitiveComponent : InPrimitives)
+	for (TSoftObjectPtr<AActor>& SoftActor : Actors)
 	{
-		TStrongObjectPtr<UPrimitiveComponent> Primitive = PrimitiveComponent.Pin();
-		if (Primitive.IsValid())
+		if (AActor* Actor = SoftActor.Get())
 		{
-			CachedVisibilityInSceneCapture.Add(Primitive.Get(), Primitive->bVisibleInSceneCaptureOnly);
+			for (UActorComponent* Component : Actor->GetComponents())
+			{
+				if (UPrimitiveComponent* PrimitiveComponent = Cast<UPrimitiveComponent>(Component))
+				{
+					PrimitiveComponent->Modify();
+					PrimitiveComponent->SetVisibleInSceneCaptureOnly(bVisibleInSceneCaptureOnly);
+				}
+			}
 		}
 	}
 }
 
-void UCompositeLayerSceneCapture::UpdatePrimitiveVisibilityState()
+void UCompositeLayerSceneCapture::RestorePrimitiveVisibilityState()
 {
+	// First, we restore current actors to class default, in case they have been deleted.
+	UPrimitiveComponent* CDO = UPrimitiveComponent::StaticClass()->GetDefaultObject<UPrimitiveComponent>();
+
 	for (TSoftObjectPtr<AActor>& SoftActor : Actors)
 	{
 		if (AActor* Actor = SoftActor.Get())
@@ -226,25 +233,73 @@ void UCompositeLayerSceneCapture::UpdatePrimitiveVisibilityState()
 				if (UPrimitiveComponent* PrimitiveComponent = Cast<UPrimitiveComponent>(Component))
 				{
 					PrimitiveComponent->Modify();
-					PrimitiveComponent->SetVisibleInSceneCaptureOnly(bVisibleInSceneCaptureOnly);
+					PrimitiveComponent->SetVisibleInSceneCaptureOnly(CDO->bVisibleInSceneCaptureOnly);
 				}
 			}
 		}
 	}
 }
 
-void UCompositeLayerSceneCapture::RestorePrimitiveVisibilityState()
+void UCompositeLayerSceneCapture::UpdateSceneCaptureComponents(USceneCaptureComponent2D& SceneCaptureComponent) const
+{
+	SceneCaptureComponent.ShowOnlyComponents.Reset(Actors.Num());
+
+	for (const TSoftObjectPtr<AActor>& SoftActor : Actors)
+	{
+		if (const AActor* Actor = SoftActor.Get())
+		{
+			for (UActorComponent* Component : Actor->GetComponents())
+			{
+				if (UPrimitiveComponent* PrimitiveComponent = Cast<UPrimitiveComponent>(Component))
+				{
+					SceneCaptureComponent.ShowOnlyComponents.Add(PrimitiveComponent);
+				}
+			}
+		}
+	}
+}
+
+void UCompositeLayerSceneCapture::OnEndOfFrameUpdate(UWorld* InWorld)
 {
-	for (const TPair<TWeakObjectPtr<UPrimitiveComponent>, bool>& Pair : CachedVisibilityInSceneCapture)
+	/**
+	* In order to retain the ability to reference actors in sublevels, we prefer using weak pointer components on scene captures. However,
+	* primitive components may often be recreated (see RerunConstructionScripts(), FStaticMeshComponentRecreateRenderStateContext, dynamic
+	* text primitives, etc). To prevent issues with destroyed primitives disappearing in our scene captures, we re-register them every frame.
+	*/
+	if (IsEnabled())
 	{
-		TStrongObjectPtr< UPrimitiveComponent> Primitive = Pair.Key.Pin();
-		if (Primitive.IsValid())
+		ACompositeActor* CompositeActor = GetTypedOuter<ACompositeActor>();
+
+		if (IsValid(CompositeActor) && CompositeActor->GetCompositeLayers().Contains(this))
 		{
-			Primitive->Modify();
-			Primitive->SetVisibleInSceneCaptureOnly(Pair.Value);
+			USceneCaptureComponent2D* SceneCaptureComponent = GetSceneCaptureComponent();
+
+			if (IsValid(SceneCaptureComponent))
+			{
+				UpdateSceneCaptureComponents(*SceneCaptureComponent);
+			}
 		}
 	}
+}
 
-	CachedVisibilityInSceneCapture.Reset();
+#if WITH_EDITOR
+void UCompositeLayerSceneCapture::OnObjectPreEditChange(UObject* Object, const FEditPropertyChain& EditPropertyChain)
+{
+	if (Object == this)
+	{
+		if (EditPropertyChain.GetActiveMemberNode())
+		{
+			const FProperty* Property = EditPropertyChain.GetActiveMemberNode()->GetValue();
+			
+			if (Property)
+			{
+				if (Property->GetName() == GET_MEMBER_NAME_CHECKED(UCompositeLayerSceneCapture, Actors))
+				{
+					RestorePrimitiveVisibilityState();
+				}
+			}
+		}
+	}
 }
+#endif
 
diff --git a/Engine/Plugins/Compositing/Composite/Source/Composite/Private/Layers/CompositeLayerShadowReflection.cpp b/Engine/Plugins/Compositing/Composite/Source/Composite/Private/Layers/CompositeLayerShadowReflection.cpp
--- a/Engine/Plugins/Compositing/Composite/Source/Composite/Private/Layers/CompositeLayerShadowReflection.cpp
+++ b/Engine/Plugins/Compositing/Composite/Source/Composite/Private/Layers/CompositeLayerShadowReflection.cpp
@@ -13,9 +13,20 @@ UCompositeLayerShadowReflection::UCompositeLayerShadowReflection(const FObjectIn
 {
 	Operation = ECompositeCoreMergeOp::Multiply;
 	AutoConfigureActors = ECompositeHiddenInSceneCaptureConfiguration::None;
+
+	if (!HasAnyFlags(RF_ClassDefaultObject))
+	{
+		// 5.7 hotfix: We must use delegates since adding virtual function overrides is not permitted
+		FWorldDelegates::OnWorldPreSendAllEndOfFrameUpdates.AddUObject(this, &UCompositeLayerShadowReflection::OnEndOfFrameUpdate);
+#if WITH_EDITOR
+		FCoreUObjectDelegates::OnPreObjectPropertyChanged.AddUObject(this, &UCompositeLayerShadowReflection::OnObjectPreEditChange);
+#endif
+	}
 }
 
+PRAGMA_DISABLE_DEPRECATION_WARNINGS
 UCompositeLayerShadowReflection::~UCompositeLayerShadowReflection() = default;
+PRAGMA_ENABLE_DEPRECATION_WARNINGS
 
 void UCompositeLayerShadowReflection::OnRemoved(const UWorld* World)
 {
@@ -26,7 +37,10 @@ void UCompositeLayerShadowReflection::OnRemoved(const UWorld* World)
 		CompositeActor->DestroySceneCaptures(this);
 	}
 
-	RestorePrimitiveVisibilityState();
+	if (AutoConfigureActors != ECompositeHiddenInSceneCaptureConfiguration::None)
+	{
+		RestorePrimitiveVisibilityState();
+	}
 }
 
 void UCompositeLayerShadowReflection::BeginDestroy()
@@ -34,6 +48,11 @@ void UCompositeLayerShadowReflection::BeginDestroy()
 	Super::BeginDestroy();
 
 	OnRemoved(GetWorld()); // Redundant remove call for safety
+
+	FWorldDelegates::OnWorldPreSendAllEndOfFrameUpdates.RemoveAll(this);
+#if WITH_EDITOR
+	FCoreUObjectDelegates::OnPreObjectPropertyChanged.RemoveAll(this);
+#endif
 }
 
 #if WITH_EDITOR
@@ -49,13 +68,25 @@ void UCompositeLayerShadowReflection::PostEditChangeProperty(FPropertyChangedEve
 	}
 	else if (PropertyName == GET_MEMBER_NAME_CHECKED(ThisClass, AutoConfigureActors))
 	{
-		UpdatePrimitiveVisibilityState();
+		if (AutoConfigureActors != ECompositeHiddenInSceneCaptureConfiguration::None)
+		{
+			const bool bHiddenInSceneCapture = (AutoConfigureActors == ECompositeHiddenInSceneCaptureConfiguration::Hidden);
+			const bool bAffectIndirectLightingWhileHidden = bHiddenInSceneCapture;
+			const bool bCastHiddenShadow = bHiddenInSceneCapture;
+
+			UpdatePrimitiveVisibilityState(bHiddenInSceneCapture, bAffectIndirectLightingWhileHidden, bCastHiddenShadow);
+		}
 	}
 }
 
 void UCompositeLayerShadowReflection::PreEditUndo()
 {
 	Super::PreEditUndo();
+
+	if (AutoConfigureActors != ECompositeHiddenInSceneCaptureConfiguration::None)
+	{
+		RestorePrimitiveVisibilityState();
+	}
 }
 
 void UCompositeLayerShadowReflection::PostEditUndo()
@@ -122,31 +153,21 @@ const TArray<TSoftObjectPtr<AActor>> UCompositeLayerShadowReflection::GetActors(
 
 void UCompositeLayerShadowReflection::SetActors(TArray<TSoftObjectPtr<AActor>> InActors)
 {
-	RestorePrimitiveVisibilityState();
+	if (AutoConfigureActors != ECompositeHiddenInSceneCaptureConfiguration::None)
+	{
+		RestorePrimitiveVisibilityState();
+	}
 
 	Actors = MoveTemp(InActors);
 
-	// Update hidden components
-	TArray<TWeakObjectPtr<UPrimitiveComponent> > HiddenComponents;
-	HiddenComponents.Reserve(Actors.Num());
-
-	for (const TSoftObjectPtr<AActor>& SoftActor : Actors)
+	if (AutoConfigureActors != ECompositeHiddenInSceneCaptureConfiguration::None)
 	{
-		if (const AActor* Actor = SoftActor.Get())
-		{
-			for (UActorComponent* Component : Actor->GetComponents())
-			{
-				if (UPrimitiveComponent* PrimitiveComponent = Cast<UPrimitiveComponent>(Component))
-				{
-					HiddenComponents.Add(PrimitiveComponent);
-				}
-			}
-		}
-	}
+		const bool bHiddenInSceneCapture = (AutoConfigureActors == ECompositeHiddenInSceneCaptureConfiguration::Hidden);
+		const bool bAffectIndirectLightingWhileHidden = bHiddenInSceneCapture;
+		const bool bCastHiddenShadow = bHiddenInSceneCapture;
 
-	CachePrimitiveVisibilityState(HiddenComponents);
-
-	UpdatePrimitiveVisibilityState();
+		UpdatePrimitiveVisibilityState(bHiddenInSceneCapture, bAffectIndirectLightingWhileHidden, bCastHiddenShadow);
+	}
 
 	ACompositeActor* CompositeActor = GetTypedOuter<ACompositeActor>();
 	if (!IsValid(CompositeActor))
@@ -167,78 +188,124 @@ void UCompositeLayerShadowReflection::SetActors(TArray<TSoftObjectPtr<AActor>> I
 
 	// First, we empty hidden actors since our previous logic used them
 	CachedSceneCaptures[1]->HiddenActors.Empty();
-	CachedSceneCaptures[1]->HiddenComponents = HiddenComponents;
+
+	UpdateSceneCaptureComponents(*CachedSceneCaptures[1]);
 }
 
-void UCompositeLayerShadowReflection::CachePrimitiveVisibilityState(TArrayView<TWeakObjectPtr<UPrimitiveComponent>> InPrimitives)
+void UCompositeLayerShadowReflection::UpdatePrimitiveVisibilityState(
+	bool bInHiddenInSceneCapture,
+	bool bInAffectIndirectLightingWhileHidden,
+	bool bInCastHiddenShadow) const
 {
-	CachedVisibilityInSceneCapture.Reset();
-
-	for (const TWeakObjectPtr<UPrimitiveComponent>& PrimitiveComponent : InPrimitives)
+	for (const TSoftObjectPtr<AActor>& Actor : Actors)
 	{
-		TStrongObjectPtr<UPrimitiveComponent> Primitive = PrimitiveComponent.Pin();
-		if (Primitive.IsValid())
+		if (Actor.IsValid())
 		{
-			FCompositeShadowReflectionPrimitiveState& State = CachedVisibilityInSceneCapture.Add(Primitive.Get());
-			State.bHiddenInSceneCapture = Primitive->bHiddenInSceneCapture;
-			State.bAffectIndirectLightingWhileHidden = Primitive->bAffectIndirectLightingWhileHidden;
-			State.bCastHiddenShadow = Primitive->bCastHiddenShadow;
+			for (UActorComponent* Component : Actor->GetComponents())
+			{
+				if (UPrimitiveComponent* PrimitiveComponent = Cast<UPrimitiveComponent>(Component))
+				{
+					bool bUpdatedComponent = false;
+
+					if (PrimitiveComponent->bHiddenInSceneCapture != bInHiddenInSceneCapture)
+					{
+						PrimitiveComponent->SetHiddenInSceneCapture(bInHiddenInSceneCapture);
+						bUpdatedComponent = true;
+					}
+
+					if (PrimitiveComponent->bAffectIndirectLightingWhileHidden != bInAffectIndirectLightingWhileHidden)
+					{
+						PrimitiveComponent->SetAffectIndirectLightingWhileHidden(bInAffectIndirectLightingWhileHidden);
+						bUpdatedComponent = true;
+					}
+
+					if (PrimitiveComponent->bCastHiddenShadow != bInCastHiddenShadow)
+					{
+						PrimitiveComponent->SetCastHiddenShadow(bInCastHiddenShadow);
+						bUpdatedComponent = true;
+					}
+
+					if (bUpdatedComponent)
+					{
+						PrimitiveComponent->Modify();
+					}
+				}
+			}
 		}
 	}
 }
 
-void UCompositeLayerShadowReflection::UpdatePrimitiveVisibilityState()
+void UCompositeLayerShadowReflection::RestorePrimitiveVisibilityState() const
 {
-	bool bHideActorsInSceneCapture = false;
+	// Restore to class default
+	UPrimitiveComponent* CDO = UPrimitiveComponent::StaticClass()->GetDefaultObject<UPrimitiveComponent>();
 
-	switch (AutoConfigureActors)
-	{
-	case ECompositeHiddenInSceneCaptureConfiguration::None:
-		return;
-	case ECompositeHiddenInSceneCaptureConfiguration::Hidden:
-		bHideActorsInSceneCapture = true;
-		break;
-	case ECompositeHiddenInSceneCaptureConfiguration::Visible:
-	default:
-		break;
-	}
+	UpdatePrimitiveVisibilityState(CDO->bHiddenInSceneCapture, CDO->bAffectIndirectLightingWhileHidden, CDO->bCastHiddenShadow);
+}
+
+void UCompositeLayerShadowReflection::UpdateSceneCaptureComponents(USceneCaptureComponent2D& SceneCaptureComponent) const
+{
+	SceneCaptureComponent.HiddenComponents.Reset(Actors.Num());
 
-	for (TSoftObjectPtr<AActor>& Actor : Actors)
+	for (const TSoftObjectPtr<AActor>& SoftActor : Actors)
 	{
-		if (Actor.IsValid())
+		if (const AActor* Actor = SoftActor.Get())
 		{
 			for (UActorComponent* Component : Actor->GetComponents())
 			{
 				if (UPrimitiveComponent* PrimitiveComponent = Cast<UPrimitiveComponent>(Component))
 				{
-					PrimitiveComponent->Modify();
-					PrimitiveComponent->SetHiddenInSceneCapture(bHideActorsInSceneCapture);
-					PrimitiveComponent->SetAffectIndirectLightingWhileHidden(bHideActorsInSceneCapture);
-					PrimitiveComponent->SetCastHiddenShadow(bHideActorsInSceneCapture);
+					SceneCaptureComponent.HiddenComponents.Add(PrimitiveComponent);
 				}
 			}
 		}
 	}
 }
 
-void UCompositeLayerShadowReflection::RestorePrimitiveVisibilityState()
+void UCompositeLayerShadowReflection::OnEndOfFrameUpdate(UWorld* InWorld)
 {
-	for (const TPair<TWeakObjectPtr<UPrimitiveComponent>, FCompositeShadowReflectionPrimitiveState>& Pair : CachedVisibilityInSceneCapture)
+	/**
+	* In order to retain the ability to reference actors in sublevels, we prefer using weak pointer components on scene captures. However,
+	* primitive components may often be recreated (see RerunConstructionScripts(), FStaticMeshComponentRecreateRenderStateContext, dynamic
+	* text primitives, etc). To prevent issues with destroyed primitives disappearing in our scene captures, we re-register them every frame.
+	*/
+	if (IsEnabled())
 	{
-		TStrongObjectPtr< UPrimitiveComponent> Primitive = Pair.Key.Pin();
-		if (Primitive.IsValid())
-		{
-			const FCompositeShadowReflectionPrimitiveState& State = Pair.Value;
+		ACompositeActor* CompositeActor = GetTypedOuter<ACompositeActor>();
 
-			Primitive->Modify();
-			Primitive->SetHiddenInSceneCapture(State.bHiddenInSceneCapture);
-			Primitive->SetAffectIndirectLightingWhileHidden(State.bAffectIndirectLightingWhileHidden);
-			Primitive->SetCastHiddenShadow(State.bCastHiddenShadow);
+		if (IsValid(CompositeActor) && CompositeActor->GetCompositeLayers().Contains(this))
+		{
+			if (CachedSceneCaptures[1].IsValid())
+			{
+				UpdateSceneCaptureComponents(*CachedSceneCaptures[1]);
+			}
 		}
 	}
+}
+
+#if WITH_EDITOR
+void UCompositeLayerShadowReflection::OnObjectPreEditChange(UObject* Object, const FEditPropertyChain& EditPropertyChain)
+{
+	if (Object == this)
+	{
+		if (EditPropertyChain.GetActiveMemberNode())
+		{
+			const FProperty* Property = EditPropertyChain.GetActiveMemberNode()->GetValue();
 
-	CachedVisibilityInSceneCapture.Reset();
+			if (Property)
+			{
+				if (Property->GetName() == GET_MEMBER_NAME_CHECKED(UCompositeLayerShadowReflection, Actors))
+				{
+					if (AutoConfigureActors != ECompositeHiddenInSceneCaptureConfiguration::None)
+					{
+						RestorePrimitiveVisibilityState();
+					}
+				}
+			}
+		}
+	}
 }
+#endif
 
 TStaticArray<TWeakObjectPtr<UCompositeShadowReflectionCatcherComponent>, 2> UCompositeLayerShadowReflection::FindOrCreateSceneCapturePair(ACompositeActor& InOuter) const
 {
diff --git a/Engine/Plugins/Compositing/Composite/Source/Composite/Private/Layers/CompositeLayerSingleLightShadow.cpp b/Engine/Plugins/Compositing/Composite/Source/Composite/Private/Layers/CompositeLayerSingleLightShadow.cpp
--- a/Engine/Plugins/Compositing/Composite/Source/Composite/Private/Layers/CompositeLayerSingleLightShadow.cpp
+++ b/Engine/Plugins/Compositing/Composite/Source/Composite/Private/Layers/CompositeLayerSingleLightShadow.cpp
@@ -162,6 +162,12 @@ UCompositeLayerSingleLightShadow::UCompositeLayerSingleLightShadow(const FObject
 	: Super(ObjectInitializer)
 {
 	Operation = ECompositeCoreMergeOp::Multiply;
+
+	if (!HasAnyFlags(RF_ClassDefaultObject))
+	{
+		// 5.7 hotfix: We must use delegates since adding virtual function overrides is not permitted
+		FWorldDelegates::OnWorldPreSendAllEndOfFrameUpdates.AddUObject(this, &UCompositeLayerSingleLightShadow::OnEndOfFrameUpdate);
+	}
 }
 
 void UCompositeLayerSingleLightShadow::OnRemoved(const UWorld* World)
@@ -179,6 +185,8 @@ void UCompositeLayerSingleLightShadow::BeginDestroy()
 	Super::BeginDestroy();
 
 	OnRemoved(GetWorld()); // Redundant remove call for safety
+
+	FWorldDelegates::OnWorldPreSendAllEndOfFrameUpdates.RemoveAll(this);
 }
 
 #if WITH_EDITOR
@@ -260,12 +268,8 @@ bool UCompositeLayerSingleLightShadow::GetProxy(FTraversalContext& InContext, FS
 		PassInputs[0].Set<FPassExternalResourceDesc>({ SceneDepthId });
 		PassInputs[1].Set<FPassExternalResourceDesc>({ ShadowMapId });
 
-		static const TConsoleVariableData<float>* CVarDepthBias = IConsoleManager::Get().FindTConsoleVariableDataFloat(TEXT("r.Shadow.CSMDepthBias"));
-		float ShadowDepthBias = CVarDepthBias ? CVarDepthBias->GetValueOnGameThread() : 10.0f;
-		ShadowDepthBias *= ShadowBias;
-
 		ShadowProxy = InFrameAllocator.Create<FCompositeSingleLightShadowProxy>(PassInputs);
-		ShadowProxy->ShadowStrengthBias = FVector4f(ShadowStrength, ShadowDepthBias, 0.0f, 0.0f);
+		ShadowProxy->ShadowStrengthBias = FVector4f(ShadowStrength, 0.0f, 0.0f, 0.0f);
 		GetShadowMatrices(*CachedShadowMapCapture, ShadowProxy->ShadowMatrix, ShadowProxy->ShadowInvDeviceZToWorldZ);
 	}
 	ShadowProxy->LightWorldDirection = FVector3f(Light->GetTransform().TransformVector(FVector(1, 0, 0)));
@@ -294,39 +298,28 @@ void UCompositeLayerSingleLightShadow::SetShadowCastingActors(TArray<TSoftObject
 		return;
 	}
 
-	if (!CachedSceneDepthCapture.IsValid())
+	if (!CachedShadowMapCapture.IsValid())
 	{
-		CachedSceneDepthCapture = FindOrCreateSceneDepthCapture(*CompositeActor);
+		CachedShadowMapCapture = FindOrCreateShadowMapCapture(*CompositeActor);
 	}
-	
-	// Early exit if the cached scene depth capture is still not valid
+
 	if (!CachedSceneDepthCapture.IsValid())
 	{
-		return;
+		CachedSceneDepthCapture = FindOrCreateSceneDepthCapture(*CompositeActor);
 	}
 
-	// First, we empty hidden actors since our previous logic used them
-	CachedSceneDepthCapture->HiddenActors.Empty();
-
-	// Collect & set hidden primitive components
-	TArray<TWeakObjectPtr<UPrimitiveComponent> > HiddenComponents;
-	HiddenComponents.Reserve(ShadowCastingActors.Num());
-
-	for (const TSoftObjectPtr<AActor>& SoftActor : ShadowCastingActors)
+	if (CachedShadowMapCapture.IsValid())
 	{
-		if (const AActor* Actor = SoftActor.Get())
-		{
-			for (UActorComponent* Component : Actor->GetComponents())
-			{
-				if (UPrimitiveComponent* PrimitiveComponent = Cast<UPrimitiveComponent>(Component))
-				{
-					HiddenComponents.Add(PrimitiveComponent);
-				}
-			}
-		}
+		UpdateSceneCaptureComponents(*CachedShadowMapCapture, ShadowCastingActors);
 	}
+	
+	if (CachedSceneDepthCapture.IsValid())
+	{
+		// First, we empty hidden actors since our previous logic used them
+		CachedSceneDepthCapture->HiddenActors.Empty();
 
-	CachedSceneDepthCapture->HiddenComponents = HiddenComponents;
+		UpdateSceneCaptureComponents(*CachedSceneDepthCapture, ShadowCastingActors);
+	}
 }
 
 UCompositeSceneCapture2DComponent* UCompositeLayerSingleLightShadow::FindOrCreateSceneDepthCapture(ACompositeActor& InOuter) const
@@ -350,7 +343,7 @@ UCompositeSceneCapture2DComponent* UCompositeLayerSingleLightShadow::FindOrCreat
 
 	if (ensure(IsValid(ShadowMapCapture)))
 	{
-		ShadowMapCapture->PrimitiveRenderMode = ESceneCapturePrimitiveRenderMode::PRM_RenderScenePrimitives;
+		ShadowMapCapture->PrimitiveRenderMode = ESceneCapturePrimitiveRenderMode::PRM_UseShowOnlyList;
 		ShadowMapCapture->CaptureSource = SCS_SceneDepth;
 		ShadowMapCapture->bRenderInMainRenderer = true;
 		ShadowMapCapture->bMainViewFamily = false;
@@ -364,6 +357,63 @@ UCompositeSceneCapture2DComponent* UCompositeLayerSingleLightShadow::FindOrCreat
 	return ShadowMapCapture;
 }
 
+void UCompositeLayerSingleLightShadow::UpdateSceneCaptureComponents(USceneCaptureComponent2D& SceneCaptureComponent, TArrayView<TSoftObjectPtr<AActor>> InActors) const
+{
+	auto UpdateComponentsFn = [InActors](TArray<TWeakObjectPtr<UPrimitiveComponent>>& InComponents) -> void
+		{
+			InComponents.Reset(InActors.Num());
+			
+			for (const TSoftObjectPtr<AActor>& SoftActor : InActors)
+			{
+				if (const AActor* Actor = SoftActor.Get())
+				{
+					for (UActorComponent* Component : Actor->GetComponents())
+					{
+						if (UPrimitiveComponent* PrimitiveComponent = Cast<UPrimitiveComponent>(Component))
+						{
+							InComponents.Add(PrimitiveComponent);
+						}
+					}
+				}
+			}
+		};
+
+	if (SceneCaptureComponent.PrimitiveRenderMode == ESceneCapturePrimitiveRenderMode::PRM_UseShowOnlyList)
+	{
+		UpdateComponentsFn(SceneCaptureComponent.ShowOnlyComponents);
+	}
+	else
+	{
+		UpdateComponentsFn(SceneCaptureComponent.HiddenComponents);
+	}
+}
+
+void UCompositeLayerSingleLightShadow::OnEndOfFrameUpdate(UWorld* InWorld)
+{
+	/**
+	* In order to retain the ability to reference actors in sublevels, we prefer using weak pointer components on scene captures. However,
+	* primitive components may often be recreated (see RerunConstructionScripts(), FStaticMeshComponentRecreateRenderStateContext, dynamic
+	* text primitives, etc). To prevent issues with destroyed primitives disappearing in our scene captures, we re-register them every frame.
+	*/
+	if (IsEnabled())
+	{
+		ACompositeActor* CompositeActor = GetTypedOuter<ACompositeActor>();
+
+		if (IsValid(CompositeActor) && CompositeActor->GetCompositeLayers().Contains(this))
+		{
+			if (CachedShadowMapCapture.IsValid())
+			{
+				UpdateSceneCaptureComponents(*CachedShadowMapCapture, ShadowCastingActors);
+			}
+
+			if (CachedSceneDepthCapture.IsValid())
+			{
+				UpdateSceneCaptureComponents(*CachedSceneDepthCapture, ShadowCastingActors);
+			}
+		}
+	}
+}
+
 void UCompositeLayerSingleLightShadow::GetShadowMatrices(UCompositeSceneCapture2DComponent& ShadowMapCapture, FMatrix44f& OutShadowMatrix, FVector4f& OutShadowInvDeviceZToWorldZ) const
 {
 	FMinimalViewInfo ViewInfo;
diff --git a/Engine/Plugins/Compositing/Composite/Source/Composite/Private/Passes/CompositePassColorGrade.cpp b/Engine/Plugins/Compositing/Composite/Source/Composite/Private/Passes/CompositePassColorGrade.cpp
--- a/Engine/Plugins/Compositing/Composite/Source/Composite/Private/Passes/CompositePassColorGrade.cpp
+++ b/Engine/Plugins/Compositing/Composite/Source/Composite/Private/Passes/CompositePassColorGrade.cpp
@@ -93,7 +93,7 @@ public:
 		Parameters->ShadowsMax = ColorGradingSettings.ShadowsMax;
 		Parameters->HighlightsMin = ColorGradingSettings.HighlightsMin;
 		Parameters->HighlightsMax = ColorGradingSettings.HighlightsMax;
-		Parameters->bIsTemperatureWhiteBalance = TemperatureSettings.TemperatureType;
+		Parameters->bIsTemperatureWhiteBalance = (TemperatureSettings.TemperatureType == ETemperatureMethod::TEMP_WhiteBalance);
 		Parameters->WhiteTemp = TemperatureSettings.WhiteTemp;
 		Parameters->WhiteTint = TemperatureSettings.WhiteTint;
 		Parameters->RenderTargets[0] = Output.GetRenderTargetBinding();
diff --git a/Engine/Plugins/Compositing/Composite/Source/Composite/Private/Passes/CompositePassColorKeyer.cpp b/Engine/Plugins/Compositing/Composite/Source/Composite/Private/Passes/CompositePassColorKeyer.cpp
--- a/Engine/Plugins/Compositing/Composite/Source/Composite/Private/Passes/CompositePassColorKeyer.cpp
+++ b/Engine/Plugins/Compositing/Composite/Source/Composite/Private/Passes/CompositePassColorKeyer.cpp
@@ -105,7 +105,7 @@ namespace UE
 				FComputeShaderUtils::AddPass(
 					GraphBuilder,
 					RDG_EVENT_NAME("CompositeCore.Denoise (%dx%d)", TextureSize.X, TextureSize.Y),
-					GSupportsEfficientAsyncCompute ? ERDGPassFlags::AsyncCompute : ERDGPassFlags::Compute,
+					ERDGPassFlags::Compute,
 					ComputeShader,
 					PassParameters,
 					FComputeShaderUtils::GetGroupCount(TextureSize, FCompositePassDenoiseShader::ThreadGroupSize)
diff --git a/Engine/Plugins/Compositing/Composite/Source/Composite/Public/Layers/CompositeLayerSceneCapture.h b/Engine/Plugins/Compositing/Composite/Source/Composite/Public/Layers/CompositeLayerSceneCapture.h
--- a/Engine/Plugins/Compositing/Composite/Source/Composite/Public/Layers/CompositeLayerSceneCapture.h
+++ b/Engine/Plugins/Compositing/Composite/Source/Composite/Public/Layers/CompositeLayerSceneCapture.h
@@ -10,6 +10,7 @@
 
 class AActor;
 class UPrimitiveComponent;
+class FEditPropertyChain;
 
 /**
 * Layer for scene capture (or custom render pass) render integration.
@@ -73,15 +74,23 @@ private:
 	/** Convenience function to access the actor-managed scene capture component. */
 	class UCompositeSceneCapture2DComponent* GetSceneCaptureComponent() const;
 
-	/** Cache scene capture visibility on registered mesh actors primitives. */
-	void CachePrimitiveVisibilityState(TArrayView<TWeakObjectPtr<UPrimitiveComponent>> InPrimitives);
-
 	/** Update scene capture visibility on registered mesh actors primitives. */
 	void UpdatePrimitiveVisibilityState();
 	
 	/** Restore scene capture visibility on registered mesh actors primitives. */
 	void RestorePrimitiveVisibilityState();
 
+	/** Update scene capture show only components. */
+	void UpdateSceneCaptureComponents(USceneCaptureComponent2D& SceneCaptureComponent) const;
+
+	/** 5.7 hotfix: End of frame callback. */
+	void OnEndOfFrameUpdate(UWorld* InWorld);
+
+#if WITH_EDITOR
+	/** 5.7 hotfix: Replacement for PreEditChange which can't be override in a hotfix. */
+	void OnObjectPreEditChange(UObject* Object, const FEditPropertyChain& EditPropertyChain);
+#endif
+
 	/**
 	 * List of actors to be rendered from the scene capture.
 	 * 
@@ -90,9 +99,11 @@ private:
 	UPROPERTY(EditAnywhere, BlueprintGetter = GetActors, BlueprintSetter = SetActors, Category = "", meta = (AllowPrivateAccess = true))
 	TArray<TSoftObjectPtr<AActor>> Actors;
 
-	/** Cached visibility in scene capture state to restore initial values. */
-	UPROPERTY()
+#if WITH_EDITORONLY_DATA
+	UE_DEPRECATED(5.7, "CachedVisibilityInSceneCapture has been deprecated in the 5.7.2 hotfix.")
+	UPROPERTY(meta = (DeprecatedProperty, DeprecationMessage = "CachedVisibilityInSceneCapture has been deprecated in the 5.7.2 hotfix."))
 	TMap<TWeakObjectPtr<UPrimitiveComponent>, bool> CachedVisibilityInSceneCapture;
+#endif
 
 	/** Visibility setting applied to registered actors. By default, registered meshes will only be visible in scene captures and no longer in the main render. */
 	UPROPERTY(EditAnywhere, BlueprintGetter = IsVisibleInSceneCaptureOnly, BlueprintSetter = SetVisibleInSceneCaptureOnly, Category = "", meta = (AllowPrivateAccess = true))
diff --git a/Engine/Plugins/Compositing/Composite/Source/Composite/Public/Layers/CompositeLayerShadowReflection.h b/Engine/Plugins/Compositing/Composite/Source/Composite/Public/Layers/CompositeLayerShadowReflection.h
--- a/Engine/Plugins/Compositing/Composite/Source/Composite/Public/Layers/CompositeLayerShadowReflection.h
+++ b/Engine/Plugins/Compositing/Composite/Source/Composite/Public/Layers/CompositeLayerShadowReflection.h
@@ -95,14 +95,23 @@ public:
 	ECompositeHiddenInSceneCaptureConfiguration AutoConfigureActors;
 
 private:
-	/** Cache scene capture visibility on registered mesh actors primitives. */
-	void CachePrimitiveVisibilityState(TArrayView<TWeakObjectPtr<UPrimitiveComponent>> InPrimitives);
 
 	/** Update scene capture visibility on registered mesh actors primitives. */
-	void UpdatePrimitiveVisibilityState();
+	void UpdatePrimitiveVisibilityState(bool bInHiddenInSceneCapture, bool bInAffectIndirectLightingWhileHidden, bool bInCastHiddenShadow) const;
 
 	/** Restore scene capture visibility on registered mesh actors primitives. */
-	void RestorePrimitiveVisibilityState();
+	void RestorePrimitiveVisibilityState() const;
+
+	/** Update scene capture hidden components. */
+	void UpdateSceneCaptureComponents(USceneCaptureComponent2D& SceneCaptureComponent) const;
+
+	/** 5.7 hotfix: End of frame callback. */
+	void OnEndOfFrameUpdate(UWorld* InWorld);
+
+#if WITH_EDITOR
+	/** 5.7 hotfix: Replacement for PreEditChange which can't be override in a hotfix. */
+	void OnObjectPreEditChange(UObject* Object, const FEditPropertyChain& EditPropertyChain);
+#endif
 
 	/** Find or create scene captures with CG and without CG. */
 	TStaticArray<TWeakObjectPtr<UCompositeShadowReflectionCatcherComponent>, 2> FindOrCreateSceneCapturePair(ACompositeActor& InOuter) const;
@@ -112,9 +121,11 @@ private:
 	UPROPERTY(EditAnywhere, BlueprintGetter = GetActors, BlueprintSetter = SetActors, Category = "")
 	TArray<TSoftObjectPtr<AActor>> Actors;
 
-	/** Cached visibility in scene capture state to restore initial values. */
-	UPROPERTY()
+#if WITH_EDITORONLY_DATA
+	UE_DEPRECATED(5.7, "CachedVisibilityInSceneCapture has been deprecated in the 5.7.2 hotfix.")
+	UPROPERTY(meta = (DeprecatedProperty, DeprecationMessage = "CachedVisibilityInSceneCapture has been deprecated in the 5.7.2 hotfix."))
 	TMap<TWeakObjectPtr<UPrimitiveComponent>, FCompositeShadowReflectionPrimitiveState> CachedVisibilityInSceneCapture;
+#endif
 
 	/** Cached pointer to the scene captures with CG and without CG. */
 	mutable TStaticArray<TWeakObjectPtr<UCompositeShadowReflectionCatcherComponent>, 2> CachedSceneCaptures;
diff --git a/Engine/Plugins/Compositing/Composite/Source/Composite/Public/Layers/CompositeLayerSingleLightShadow.h b/Engine/Plugins/Compositing/Composite/Source/Composite/Public/Layers/CompositeLayerSingleLightShadow.h
--- a/Engine/Plugins/Compositing/Composite/Source/Composite/Public/Layers/CompositeLayerSingleLightShadow.h
+++ b/Engine/Plugins/Compositing/Composite/Source/Composite/Public/Layers/CompositeLayerSingleLightShadow.h
@@ -77,7 +77,7 @@ public:
 	float OrthographicWidth = DEFAULT_ORTHOWIDTH;
 
 	/**
-	* Resolution of the shadow map texture. Lower resolutions will cause more blurry shadows, but require a higher shadow bias.
+	* Resolution of the shadow map texture. Lower resolutions will cause more blurry shadows.
 	*/
 	UPROPERTY(EditAnywhere, BlueprintReadWrite, Category = "", meta = (ClampMin = "128", UIMin = "128"))
 	int32 ShadowMapResolution = 2048;
@@ -86,11 +86,11 @@ public:
 	UPROPERTY(EditAnywhere, BlueprintReadWrite, Category = "", meta = (ClampMin = "0.0", ClampMax = "1.0", UIMin = "0.0", UIMax = "1.0"))
 	float ShadowStrength = 0.5f;
 
-	/**
-	* Shadow bias to avoid shadow acne. Increase until unwanted shadowing disappears (automatically scaled by r.Shadow.CSMDepthBias console variable).
-	*/
-	UPROPERTY(EditAnywhere, BlueprintReadWrite, Category = "", meta = (ClampMin = "0.0", UIMin = "0.0"))
-	float ShadowBias = 0.5f;
+#if WITH_EDITORONLY_DATA
+	UE_DEPRECATED(5.7, "ShadowBias has been deprecated in the 5.7.2 hotfix.")
+	UPROPERTY(meta = (DeprecatedProperty, DeprecationMessage = "ShadowBias has been deprecated in the 5.7.2 hotfix."))
+	float ShadowBias_DEPRECATED = 0.5f;
+#endif
 
 private:
 	/** Find or create scene depth capture component, correctly configured. */
@@ -99,14 +99,17 @@ private:
 	/** Find or create shadow map capture component, correctly configured. */
 	UCompositeSceneCapture2DComponent* FindOrCreateShadowMapCapture(ACompositeActor& InOuter) const;
 
+	/** Convenience function to update scene capture primitive components, depending on the scene capture primitive render mode. */
+	void UpdateSceneCaptureComponents(USceneCaptureComponent2D& SceneCaptureComponent, TArrayView<TSoftObjectPtr<AActor>> InActors) const;
+
+	/** 5.7 hotfix: End of frame callback. */
+	void OnEndOfFrameUpdate(UWorld* InWorld);
+
 	/** Convenience function to calculate shadow matrices. */
 	void GetShadowMatrices(UCompositeSceneCapture2DComponent& ShadowMapCapture, FMatrix44f& OutShadowMatrix, FVector4f& OutShadowInvDeviceZToWorldZ) const;
 
 private:
-	/**
-	* Optional: List of shadow casting actors to be removed from the main camera view's scene depth.
-	* This helps create a cleaner shadow matte preventing aliased dark edges around CG objects.
-	*/
+	/** List of shadow casting actors (required). */
 	UPROPERTY(EditAnywhere, BlueprintGetter = GetShadowCastingActors, BlueprintSetter = SetShadowCastingActors, Category = "", meta = (AllowPrivateAccess = true))
 	TArray<TSoftObjectPtr<AActor>> ShadowCastingActors;
 
diff --git a/Engine/Plugins/Compositing/CompositeCore/Source/CompositeCore/Private/CompositeCoreSceneViewExtension.cpp b/Engine/Plugins/Compositing/CompositeCore/Source/CompositeCore/Private/CompositeCoreSceneViewExtension.cpp
--- a/Engine/Plugins/Compositing/CompositeCore/Source/CompositeCore/Private/CompositeCoreSceneViewExtension.cpp
+++ b/Engine/Plugins/Compositing/CompositeCore/Source/CompositeCore/Private/CompositeCoreSceneViewExtension.cpp
@@ -457,10 +457,7 @@ bool FCompositeCoreSceneViewExtension::IsActiveForViewFamily(const FSceneViewFam
 
 bool FCompositeCoreSceneViewExtension::IsActiveForView(const FSceneView& InView) const
 {
-	bool bIsActive = IsActiveForViewFamily(*InView.Family);
-	bIsActive &= CustomRenderTargetPerView_RenderThread.Contains(InView.GetViewKey()) || !ExternalInputs_RenderThread.IsEmpty();
-
-	return bIsActive;
+	return IsActiveForViewFamily(*InView.Family);
 }
 
 void FCompositeCoreSceneViewExtension::ApplyPreprocessing(FRDGBuilder& GraphBuilder, FSceneView& InView) const
@@ -545,20 +542,30 @@ UE::CompositeCore::FScopedExternalTextureMap::FScopedExternalTextureMap(const FC
 		Resource.Texture = FScreenPassTexture{ GSystemTextures.GetBlackDummy(GraphBuilder) };
 	}
 
-	const TRefCountPtr<IPooledRenderTarget>* CompositeRenderPassPtr = SVE.CustomRenderTargetPerView_RenderThread.Find(InView.GetViewKey());
-	if (CompositeRenderPassPtr != nullptr)
 	{
 		FPassInput& Resource = Textures.Add(BUILT_IN_CRP_ID);
-		Resource.Texture = FScreenPassTexture{ GraphBuilder.RegisterExternalTexture(*CompositeRenderPassPtr) };
-		Resource.Metadata.bInvertedAlpha = true;
+
+		const TRefCountPtr<IPooledRenderTarget>* CompositeRenderPassPtr = SVE.CustomRenderTargetPerView_RenderThread.Find(InView.GetViewKey());
+		if (CompositeRenderPassPtr != nullptr)
+		{
+			Resource.Texture = FScreenPassTexture{ GraphBuilder.RegisterExternalTexture(*CompositeRenderPassPtr) };
+			Resource.Metadata.bInvertedAlpha = true;
+		}
+		else
+		{
+			// Still provide an empty CRP target
+			Resource.Texture = FScreenPassTexture{ GSystemTextures.GetBlackDummy(GraphBuilder) };
+		}
 	}
 
 	RestoreToExternal.Reserve(SVE.ExternalInputs_RenderThread.Num());
 
 	for (const auto& Pair : SVE.ExternalInputs_RenderThread)
 	{
+		const TRefCountPtr<IPooledRenderTarget>& PooledRenderTarget = Pair.Value.RenderTarget;
+		FRDGTextureRef TextureResource = GraphBuilder.RegisterExternalTexture(PooledRenderTarget);
+		
 		FPassInput& Resource = Textures.Add(Pair.Key);
-		FRDGTextureRef TextureResource = GraphBuilder.RegisterExternalTexture(Pair.Value.RenderTarget);
 		Resource.Texture = FScreenPassTexture{ TextureResource };
 		Resource.Metadata = Pair.Value.Metadata;
 
@@ -572,7 +579,7 @@ UE::CompositeCore::FScopedExternalTextureMap::FScopedExternalTextureMap(const FC
 		 * GraphBuilder.UseExternalAccessMode(Buffer, AccessModeState.Access, AccessModeState.Pipelines); use in FRDGResourceDumpContext.
 		 */
 		FRHITexture* FamilyRenderTargetRHI = InView.Family->RenderTarget ? InView.Family->RenderTarget->GetRenderTargetTexture().GetReference() : nullptr;
-		if (TextureResource && (TextureResource->GetRHI() != FamilyRenderTargetRHI))
+		if (TextureResource && (PooledRenderTarget->GetRHI() != FamilyRenderTargetRHI))
 		{
 			GraphBuilder.UseInternalAccessMode(TextureResource);
 
