Composure: Fixed issue where UI would not display the correct resolution on a composite actor in the Composure panel.

diff --git a/Engine/Plugins/Compositing/Composite/Source/CompositeEditor/Private/UI/SCompositeEditorPanel.cpp b/Engine/Plugins/Compositing/Composite/Source/CompositeEditor/Private/UI/SCompositeEditorPanel.cpp
--- a/Engine/Plugins/Compositing/Composite/Source/CompositeEditor/Private/UI/SCompositeEditorPanel.cpp
+++ b/Engine/Plugins/Compositing/Composite/Source/CompositeEditor/Private/UI/SCompositeEditorPanel.cpp
@@ -60,6 +60,23 @@ public:
 	{
 		PropertyHandle = InPropertyHandle;
 
+		XPropertyHandle = InPropertyHandle->GetChildHandle(GET_MEMBER_NAME_CHECKED(FIntPoint, X));
+		YPropertyHandle = InPropertyHandle->GetChildHandle(GET_MEMBER_NAME_CHECKED(FIntPoint, Y));
+
+		FIntPoint InitialResolution = FIntPoint::ZeroValue;
+		FPropertyAccess::Result XResult = XPropertyHandle->GetValue(InitialResolution.X);
+		FPropertyAccess::Result YResult = YPropertyHandle->GetValue(InitialResolution.Y);
+
+		int32 InitialResolutionIndex = NamedResolutions.IndexOfByPredicate([&InitialResolution](const FNamedResolutionPtr& NamedResolution)
+		{
+			return NamedResolution->Resolution == InitialResolution;
+		});
+
+		if (InitialResolutionIndex == INDEX_NONE)
+		{
+			InitialResolutionIndex = NamedResolutions.Num() - 1;
+		}
+		
 		InHeaderRow
 		.NameContent()
 		[
@@ -77,7 +94,7 @@ public:
 					.ToolTipText(InItem->ToolTip);
 			})
 			.OnSelectionChanged(this, &FResolutionTypeCustomization::SetResolution)
-			.InitiallySelectedItem(NamedResolutions[1])
+			.InitiallySelectedItem(NamedResolutions[InitialResolutionIndex])
 			[
 				SNew(STextBlock)
 				.Font(IDetailLayoutBuilder::GetDetailFont())
@@ -97,9 +114,6 @@ public:
 
 	virtual void CustomizeChildren(TSharedRef<IPropertyHandle> InPropertyHandle, IDetailChildrenBuilder& InChildBuilder, IPropertyTypeCustomizationUtils& InCustomizationUtils) override
 	{
-		XPropertyHandle = InPropertyHandle->GetChildHandle(GET_MEMBER_NAME_CHECKED(FIntPoint, X));
-		YPropertyHandle = InPropertyHandle->GetChildHandle(GET_MEMBER_NAME_CHECKED(FIntPoint, Y));
-
 		InChildBuilder.AddCustomRow(InPropertyHandle->GetPropertyDisplayName())
 		.NameContent()
 		[
