Composure: Avoid checked module loading on destruction.

diff --git a/Engine/Plugins/Compositing/Composite/Source/Composite/Private/CompositeActor.cpp b/Engine/Plugins/Compositing/Composite/Source/Composite/Private/CompositeActor.cpp
--- a/Engine/Plugins/Compositing/Composite/Source/Composite/Private/CompositeActor.cpp
+++ b/Engine/Plugins/Compositing/Composite/Source/Composite/Private/CompositeActor.cpp
@@ -821,7 +821,10 @@ void ACompositeActor::DestroySceneCaptures(const UCompositeLayerBase* InLayer)
 #if WITH_EDITOR
 	if (bComponentRemoved)
 	{
-		FModuleManager::LoadModuleChecked<FLevelEditorModule>("LevelEditor").BroadcastComponentsEdited();
+		if (FLevelEditorModule* LevelEditorModule = FModuleManager::GetModulePtr<FLevelEditorModule>("LevelEditor"))
+		{
+			LevelEditorModule->BroadcastComponentsEdited();
+		}
 	}
 #endif
 }
