Composure: Fixed issue where copy/paste of passes and layers on composite actor were incorrectly transient and would not be saved

diff --git a/Engine/Plugins/Compositing/Composite/Source/CompositeEditor/Private/UI/SCompositePanelLayerTree.cpp b/Engine/Plugins/Compositing/Composite/Source/CompositeEditor/Private/UI/SCompositePanelLayerTree.cpp
--- a/Engine/Plugins/Compositing/Composite/Source/CompositeEditor/Private/UI/SCompositePanelLayerTree.cpp
+++ b/Engine/Plugins/Compositing/Composite/Source/CompositeEditor/Private/UI/SCompositePanelLayerTree.cpp
@@ -2299,7 +2299,7 @@ void SCompositePanelLayerTree::PasteSelectedItems()
 	FPlatformApplicationMisc::ClipboardPaste(ClipboardContent);
 	
 	FCompositeLayerObjectTextFactory Factory;
-	Factory.ProcessBuffer(GetTransientPackage(), RF_Transactional | RF_Transient, ClipboardContent);
+	Factory.ProcessBuffer(GetTransientPackage(), RF_Transactional, ClipboardContent);
 
 	if (Factory.CompositeLayers.IsEmpty())
 	{
diff --git a/Engine/Plugins/Compositing/Composite/Source/CompositeEditor/Private/UI/SCompositePassTree.cpp b/Engine/Plugins/Compositing/Composite/Source/CompositeEditor/Private/UI/SCompositePassTree.cpp
--- a/Engine/Plugins/Compositing/Composite/Source/CompositeEditor/Private/UI/SCompositePassTree.cpp
+++ b/Engine/Plugins/Compositing/Composite/Source/CompositeEditor/Private/UI/SCompositePassTree.cpp
@@ -1737,7 +1737,7 @@ void SCompositePassTree::PasteSelectedItems()
 	FPlatformApplicationMisc::ClipboardPaste(ClipboardContent);
 	
 	FCompositePassObjectTextFactory Factory;
-	Factory.ProcessBuffer(GetTransientPackage(), RF_Transactional | RF_Transient, ClipboardContent);
+	Factory.ProcessBuffer(GetTransientPackage(), RF_Transactional, ClipboardContent);
 
 	if (Factory.CompositePasses.IsEmpty())
 	{
