Composure: Scene capture layer fixes & improvements.

* We now re-register primitives on end-of-frame, to prevent loss of destroyed & recreated primitives.
* Adjusted scene capture defaults: disable sky atmosphere deferred pass to preserve atmospheric lighting, enable anti-aliasing by default.
* Fixed destruction of scene capture components, which now failed to occur since our UI correctly detaches the layer objects on deletion.

diff --git a/Engine/Plugins/Compositing/Composite/Source/Composite/Private/Components/CompositeSceneCapture2DComponent.cpp b/Engine/Plugins/Compositing/Composite/Source/Composite/Private/Components/CompositeSceneCapture2DComponent.cpp
--- a/Engine/Plugins/Compositing/Composite/Source/Composite/Private/Components/CompositeSceneCapture2DComponent.cpp
+++ b/Engine/Plugins/Compositing/Composite/Source/Composite/Private/Components/CompositeSceneCapture2DComponent.cpp
@@ -43,13 +43,18 @@ UCompositeSceneCapture2DComponent::UCompositeSceneCapture2DComponent(const FObje
 	PostProcessSettings.bOverride_LumenRayLightingMode = true;
 	PostProcessSettings.LumenRayLightingMode = ELumenRayLightingModeOverride::HitLighting;
 
-	ShowFlags.SetAtmosphere(false);
+	ShowFlags.SetDeferredAtmospherePass(false);
+	ShowFlags.SetCloud(false);
 	ShowFlags.SetAllowPrimitiveAlphaHoldout(false);
 	ShowFlags.SetBloom(false);
 	ShowFlags.SetEyeAdaptation(false);
 	ShowFlags.SetLocalExposure(false);
 	ShowFlags.SetPostProcessMaterial(false);
-
+	/**
+	 * This can cause temporal ghosting artifacts on moving shadows, but remains preferable over the otherwise default shadow noise.
+	 * Automtically ignored with custom render passes, where post-processing does not apply.
+	 */
+	ShowFlags.SetTemporalAA(true);
 }
 
 void UCompositeSceneCapture2DComponent::BeginDestroy()
diff --git a/Engine/Plugins/Compositing/Composite/Source/Composite/Private/Components/CompositeShadowReflectionCatcherComponent.cpp b/Engine/Plugins/Compositing/Composite/Source/Composite/Private/Components/CompositeShadowReflectionCatcherComponent.cpp
--- a/Engine/Plugins/Compositing/Composite/Source/Composite/Private/Components/CompositeShadowReflectionCatcherComponent.cpp
+++ b/Engine/Plugins/Compositing/Composite/Source/Composite/Private/Components/CompositeShadowReflectionCatcherComponent.cpp
@@ -13,10 +13,9 @@ UCompositeShadowReflectionCatcherComponent::UCompositeShadowReflectionCatcherCom
 {
 	PrimitiveRenderMode = ESceneCapturePrimitiveRenderMode::PRM_RenderScenePrimitives;
 
+	ShowFlags.SetAtmosphere(false);
 	ShowFlags.SetGrain(false);
 	ShowFlags.SetScreenSpaceReflections(false);
-	//Note: This can cause temporal ghosting artifacts on moving shadows, but remains preferable over the otherwise default shadow noise.
-	ShowFlags.SetTemporalAA(true);
 }
 
 #undef LOCTEXT_NAMESPACE
diff --git a/Engine/Plugins/Compositing/Composite/Source/Composite/Private/CompositeActor.cpp b/Engine/Plugins/Compositing/Composite/Source/Composite/Private/CompositeActor.cpp
--- a/Engine/Plugins/Compositing/Composite/Source/Composite/Private/CompositeActor.cpp
+++ b/Engine/Plugins/Compositing/Composite/Source/Composite/Private/CompositeActor.cpp
@@ -395,7 +395,7 @@ void ACompositeActor::Destroyed()
 		// Note: We don't use IsValid(..) since destruction should proceed even if the actor is pending kill.
 		if (Layer != nullptr)
 		{
-			Layer->OnRemoved(GetWorld());
+			Layer->OnRemoved(this);
 		}
 	}
 }
@@ -501,7 +501,7 @@ void ACompositeActor::PostEditChangeChainProperty(struct FPropertyChangedChainEv
 				{
 					if (!CompositeLayers.Contains(PreEditCompositeLayer))
 					{
-						PreEditCompositeLayer->OnRemoved(GetWorld());
+						PreEditCompositeLayer->OnRemoved(this);
 					}
 				}
 			}
@@ -547,7 +547,7 @@ void ACompositeActor::PostEditUndo()
 		{
 			if (!CompositeLayers.Contains(PreEditCompositeLayer))
 			{
-				PreEditCompositeLayer->OnRemoved(GetWorld());
+				PreEditCompositeLayer->OnRemoved(this);
 			}
 		}
 	}
@@ -687,7 +687,7 @@ void ACompositeActor::SetCompositeLayers(TArray<UCompositeLayerBase*> InLayers)
 		{
 			if (!InLayers.Contains(PreviousLayer))
 			{
-				PreviousLayer->OnRemoved(GetWorld());
+				PreviousLayer->OnRemoved(this);
 			}
 		}
 	}
diff --git a/Engine/Plugins/Compositing/Composite/Source/Composite/Private/Layers/CompositeLayerBase.cpp b/Engine/Plugins/Compositing/Composite/Source/Composite/Private/Layers/CompositeLayerBase.cpp
--- a/Engine/Plugins/Compositing/Composite/Source/Composite/Private/Layers/CompositeLayerBase.cpp
+++ b/Engine/Plugins/Compositing/Composite/Source/Composite/Private/Layers/CompositeLayerBase.cpp
@@ -25,7 +25,11 @@ UCompositeLayerBase::UCompositeLayerBase(const FObjectInitializer& ObjectInitial
 #endif
 }
 
-UCompositeLayerBase::~UCompositeLayerBase() = default;
+UCompositeLayerBase::~UCompositeLayerBase()
+{
+	constexpr bool bEnabled = false;
+	RegisterEndOfFrameUpdate(bEnabled);
+}
 
 
 #if WITH_EDITOR
@@ -67,6 +71,12 @@ void UCompositeLayerBase::PostEditChangeProperty(FPropertyChangedEvent& Property
 }
 #endif //WITH_EDITOR
 
+void UCompositeLayerBase::OnRemoved(ACompositeActor* LastOwner)
+{
+	constexpr bool bEnabled = false;
+	RegisterEndOfFrameUpdate(bEnabled);
+}
+
 bool UCompositeLayerBase::IsEnabled() const
 {
 	return bIsEnabled;
@@ -77,6 +87,25 @@ void UCompositeLayerBase::SetEnabled(bool bInEnabled)
 	bIsEnabled = bInEnabled;
 }
 
+void UCompositeLayerBase::RegisterEndOfFrameUpdate(bool bInEnabled)
+{
+	if (bInEnabled)
+	{
+		if (!OnWorldPreSendAllEndOfFrameUpdatesHandle.IsValid())
+		{
+			OnWorldPreSendAllEndOfFrameUpdatesHandle = FWorldDelegates::OnWorldPreSendAllEndOfFrameUpdates.AddUObject(this, &UCompositeLayerBase::OnEndOfFrameUpdate);
+		}
+	}
+	else
+	{
+		if (OnWorldPreSendAllEndOfFrameUpdatesHandle.IsValid())
+		{
+			FWorldDelegates::OnWorldPreSendAllEndOfFrameUpdates.Remove(OnWorldPreSendAllEndOfFrameUpdatesHandle);
+			OnWorldPreSendAllEndOfFrameUpdatesHandle.Reset();
+		}
+	}
+}
+
 void UCompositeLayerBase::AddChildPasses(UE::CompositeCore::FPassInputDecl& InBasePassInput, FTraversalContext& InContext, FSceneRenderingBulkObjectAllocator& InFrameAllocator, TArrayView<const TObjectPtr<UCompositePassBase>> InPasses) const
 {
 	// Iterate backwards so that lower passes in the UI are executed first
diff --git a/Engine/Plugins/Compositing/Composite/Source/Composite/Private/Layers/CompositeLayerPlate.cpp b/Engine/Plugins/Compositing/Composite/Source/Composite/Private/Layers/CompositeLayerPlate.cpp
--- a/Engine/Plugins/Compositing/Composite/Source/Composite/Private/Layers/CompositeLayerPlate.cpp
+++ b/Engine/Plugins/Compositing/Composite/Source/Composite/Private/Layers/CompositeLayerPlate.cpp
@@ -116,10 +116,7 @@ void UCompositeLayerPlate::Tick(float DeltaTime)
 	if (IsValid(Subsystem))
 	{
 		// TODO: We will need to separate the holdout feature from the custom render pass in the scene view extension, so that one can be used without the other.
-		if (UE::CompositeCore::IsRegisterPrimitivesOnTickEnabled())
-		{
-			Subsystem->RegisterPrimitives(GetPrimitives());
-		}
+		Subsystem->RegisterPrimitives(GetPrimitives());
 	}
 
 	{
@@ -191,15 +188,17 @@ void UCompositeLayerPlate::PostLoad()
 	}
 }
 
-void UCompositeLayerPlate::OnRemoved(const UWorld* World)
+void UCompositeLayerPlate::OnRemoved(ACompositeActor* LastOwner)
 {
+	Super::OnRemoved(LastOwner);
+
 	FCompositeRenderTargetPool::Get().ReleaseTarget(MediaRenderTarget);
 	FCompositeRenderTargetPool::Get().ReleaseTarget(SceneRenderTarget);
 
 	TryCloseMediaProfileSource();
 
 	constexpr bool bApply = false;
-	PropagateStateChange(bApply, World);
+	PropagateStateChange(bApply, LastOwner ? LastOwner->GetWorld() : nullptr);
 }
 
 void UCompositeLayerPlate::OnRenderingStateChange(bool bApply)
@@ -220,7 +219,7 @@ void UCompositeLayerPlate::BeginDestroy()
 	}
 #endif // WITH_EDITOR
 
-	OnRemoved(GetWorld()); // Redundant remove call for safety
+	OnRemoved(GetTypedOuter<ACompositeActor>()); // Redundant remove call for safety
 }
 
 const TArray<TSoftObjectPtr<AActor>> UCompositeLayerPlate::GetCompositeMeshes() const
diff --git a/Engine/Plugins/Compositing/Composite/Source/Composite/Private/Layers/CompositeLayerSceneCapture.cpp b/Engine/Plugins/Compositing/Composite/Source/Composite/Private/Layers/CompositeLayerSceneCapture.cpp
--- a/Engine/Plugins/Compositing/Composite/Source/Composite/Private/Layers/CompositeLayerSceneCapture.cpp
+++ b/Engine/Plugins/Compositing/Composite/Source/Composite/Private/Layers/CompositeLayerSceneCapture.cpp
@@ -15,27 +15,36 @@ UCompositeLayerSceneCapture::UCompositeLayerSceneCapture(const FObjectInitialize
 	, bCustomRenderPass{false}
 {
 	Operation = ECompositeCoreMergeOp::Over;
+
+	if (!HasAnyFlags(RF_ClassDefaultObject))
+	{
+		constexpr bool bEnabled = true;
+		RegisterEndOfFrameUpdate(bEnabled);
+	}
 }
 
 UCompositeLayerSceneCapture::~UCompositeLayerSceneCapture() = default;
 
-void UCompositeLayerSceneCapture::OnRemoved(const UWorld* World)
+void UCompositeLayerSceneCapture::OnRemoved(ACompositeActor* LastOwner)
 {
-	ACompositeActor* CompositeActor = GetTypedOuter<ACompositeActor>();
+	Super::OnRemoved(LastOwner);
+
 	// Note: We don't use IsValid(..) since destruction should proceed even if the actor is pending kill.
-	if (CompositeActor != nullptr)
+	if (LastOwner != nullptr)
 	{
-		CompositeActor->DestroySceneCaptures(this);
+		LastOwner->DestroySceneCaptures(this);
 	}
 
-	RestorePrimitiveVisibilityState();
+	// Restore to class default
+	UPrimitiveComponent* CDO = UPrimitiveComponent::StaticClass()->GetDefaultObject<UPrimitiveComponent>();
+	UpdatePrimitiveVisibilityState(CDO->bVisibleInSceneCaptureOnly);
 }
 
 void UCompositeLayerSceneCapture::BeginDestroy()
 {
 	Super::BeginDestroy();
 
-	OnRemoved(GetWorld()); // Redundant remove call for safety
+	OnRemoved(GetTypedOuter<ACompositeActor>()); // Redundant remove call for safety
 }
 
 #if WITH_EDITOR
@@ -74,6 +83,20 @@ void UCompositeLayerSceneCapture::PostEditUndo()
 }
 #endif
 
+void UCompositeLayerSceneCapture::OnEndOfFrameUpdate(UWorld* InWorld)
+{
+	/**
+	* In order to retain the ability to reference actors in sublevels, we prefer using weak pointer components on scene captures. However,
+	* primitive components may often be recreated (see RerunConstructionScripts(), FStaticMeshComponentRecreateRenderStateContext, dynamic
+	* text primitives, etc). To prevent issues with destroyed primitives disappearing in our scene captures, we re-register them every frame.
+	*/
+	USceneCaptureComponent2D* SceneCaptureComponent = GetSceneCaptureComponent();
+	if (IsValid(SceneCaptureComponent))
+	{
+		UpdateSceneCaptureComponents(*SceneCaptureComponent);
+	}
+}
+
 bool UCompositeLayerSceneCapture::GetProxy(FTraversalContext& InContext, FSceneRenderingBulkObjectAllocator& InFrameAllocator, FCompositeCorePassProxy*& OutProxy) const
 {
 	using namespace UE::CompositeCore;
@@ -119,38 +142,18 @@ const TArray<TSoftObjectPtr<AActor>> UCompositeLayerSceneCapture::GetActors() co
 
 void UCompositeLayerSceneCapture::SetActors(TArray<TSoftObjectPtr<AActor>> InActors)
 {
-	TArray<TWeakObjectPtr<UPrimitiveComponent>> ShowOnlyComponents;
-	ShowOnlyComponents.Reserve(InActors.Num());
+	Actors = MoveTemp(InActors);
 
-	for (const TSoftObjectPtr<AActor>& SoftActor : InActors)
-	{
-		if (const AActor* Actor = SoftActor.Get())
-		{
-			for (UActorComponent* Component : Actor->GetComponents())
-			{
-				if (UPrimitiveComponent* PrimitiveComponent = Cast<UPrimitiveComponent>(Component))
-				{
-					ShowOnlyComponents.Add(PrimitiveComponent);
-				}
-			}
-		}
-	}
+	UpdatePrimitiveVisibilityState(bVisibleInSceneCaptureOnly);
 
 	USceneCaptureComponent2D* SceneCaptureComponent = GetSceneCaptureComponent();
 	if (IsValid(SceneCaptureComponent))
 	{
 		// First, we empty show-only actors since our previous logic used them
 		SceneCaptureComponent->ShowOnlyActors.Empty();
-		SceneCaptureComponent->ShowOnlyComponents = ShowOnlyComponents;
-	}
-
-	RestorePrimitiveVisibilityState();
-
-	Actors = MoveTemp(InActors);
-
-	CachePrimitiveVisibilityState(ShowOnlyComponents);
 
-	UpdatePrimitiveVisibilityState();
+		UpdateSceneCaptureComponents(*SceneCaptureComponent);
+	}
 }
 
 bool UCompositeLayerSceneCapture::IsCustomRenderPass() const
@@ -187,7 +190,7 @@ void UCompositeLayerSceneCapture::SetVisibleInSceneCaptureOnly(bool bInVisible)
 {
 	bVisibleInSceneCaptureOnly = bInVisible;
 
-	UpdatePrimitiveVisibilityState();
+	UpdatePrimitiveVisibilityState(bVisibleInSceneCaptureOnly);
 }
 
 UCompositeSceneCapture2DComponent* UCompositeLayerSceneCapture::GetSceneCaptureComponent() const
@@ -201,50 +204,43 @@ UCompositeSceneCapture2DComponent* UCompositeLayerSceneCapture::GetSceneCaptureC
 	return nullptr;
 }
 
-void UCompositeLayerSceneCapture::CachePrimitiveVisibilityState(TArrayView<TWeakObjectPtr<UPrimitiveComponent>> InPrimitives)
-{
-	CachedVisibilityInSceneCapture.Reset();
-
-	for (const TWeakObjectPtr<UPrimitiveComponent>& PrimitiveComponent : InPrimitives)
-	{
-		TStrongObjectPtr<UPrimitiveComponent> Primitive = PrimitiveComponent.Pin();
-		if (Primitive.IsValid())
-		{
-			CachedVisibilityInSceneCapture.Add(Primitive.Get(), Primitive->bVisibleInSceneCaptureOnly);
-		}
-	}
-}
-
-void UCompositeLayerSceneCapture::UpdatePrimitiveVisibilityState()
+void UCompositeLayerSceneCapture::UpdatePrimitiveVisibilityState(bool bInVisibleInSceneCaptureOnly) const
 {
-	for (TSoftObjectPtr<AActor>& SoftActor : Actors)
+	for (const TSoftObjectPtr<AActor>& SoftActor : Actors)
 	{
-		if (AActor* Actor = SoftActor.Get())
+		if (const AActor* Actor = SoftActor.Get())
 		{
 			for (UActorComponent* Component : Actor->GetComponents())
 			{
 				if (UPrimitiveComponent* PrimitiveComponent = Cast<UPrimitiveComponent>(Component))
 				{
-					PrimitiveComponent->Modify();
-					PrimitiveComponent->SetVisibleInSceneCaptureOnly(bVisibleInSceneCaptureOnly);
+					if (PrimitiveComponent->bVisibleInSceneCaptureOnly != bInVisibleInSceneCaptureOnly)
+					{
+						PrimitiveComponent->Modify();
+						PrimitiveComponent->SetVisibleInSceneCaptureOnly(bInVisibleInSceneCaptureOnly);
+					}
 				}
 			}
 		}
 	}
 }
 
-void UCompositeLayerSceneCapture::RestorePrimitiveVisibilityState()
+void UCompositeLayerSceneCapture::UpdateSceneCaptureComponents(USceneCaptureComponent2D& SceneCaptureComponent) const
 {
-	for (const TPair<TWeakObjectPtr<UPrimitiveComponent>, bool>& Pair : CachedVisibilityInSceneCapture)
+	SceneCaptureComponent.ShowOnlyComponents.Reset(Actors.Num());
+
+	for (const TSoftObjectPtr<AActor>& SoftActor : Actors)
 	{
-		TStrongObjectPtr< UPrimitiveComponent> Primitive = Pair.Key.Pin();
-		if (Primitive.IsValid())
+		if (const AActor* Actor = SoftActor.Get())
 		{
-			Primitive->Modify();
-			Primitive->SetVisibleInSceneCaptureOnly(Pair.Value);
+			for (UActorComponent* Component : Actor->GetComponents())
+			{
+				if (UPrimitiveComponent* PrimitiveComponent = Cast<UPrimitiveComponent>(Component))
+				{
+					SceneCaptureComponent.ShowOnlyComponents.Add(PrimitiveComponent);
+				}
+			}
 		}
 	}
-
-	CachedVisibilityInSceneCapture.Reset();
 }
 
diff --git a/Engine/Plugins/Compositing/Composite/Source/Composite/Private/Layers/CompositeLayerShadowReflection.cpp b/Engine/Plugins/Compositing/Composite/Source/Composite/Private/Layers/CompositeLayerShadowReflection.cpp
--- a/Engine/Plugins/Compositing/Composite/Source/Composite/Private/Layers/CompositeLayerShadowReflection.cpp
+++ b/Engine/Plugins/Compositing/Composite/Source/Composite/Private/Layers/CompositeLayerShadowReflection.cpp
@@ -13,27 +13,40 @@ UCompositeLayerShadowReflection::UCompositeLayerShadowReflection(const FObjectIn
 {
 	Operation = ECompositeCoreMergeOp::Multiply;
 	AutoConfigureActors = ECompositeHiddenInSceneCaptureConfiguration::None;
+
+	if (!HasAnyFlags(RF_ClassDefaultObject))
+	{
+		constexpr bool bEnabled = true;
+		RegisterEndOfFrameUpdate(bEnabled);
+	}
 }
 
 UCompositeLayerShadowReflection::~UCompositeLayerShadowReflection() = default;
 
-void UCompositeLayerShadowReflection::OnRemoved(const UWorld* World)
+void UCompositeLayerShadowReflection::OnRemoved(ACompositeActor* LastOwner)
 {
-	ACompositeActor* CompositeActor = GetTypedOuter<ACompositeActor>();
+	Super::OnRemoved(LastOwner);
+
 	// Note: We don't use IsValid(..) since destruction should proceed even if the actor is pending kill.
-	if (CompositeActor != nullptr)
+	if (LastOwner != nullptr)
 	{
-		CompositeActor->DestroySceneCaptures(this);
+		LastOwner->DestroySceneCaptures(this);
 	}
 
-	RestorePrimitiveVisibilityState();
+	if (AutoConfigureActors != ECompositeHiddenInSceneCaptureConfiguration::None)
+	{
+		// Restore to class default
+		UPrimitiveComponent* CDO = UPrimitiveComponent::StaticClass()->GetDefaultObject<UPrimitiveComponent>();
+
+		UpdatePrimitiveVisibilityState(CDO->bHiddenInSceneCapture, CDO->bAffectIndirectLightingWhileHidden, CDO->bCastHiddenShadow);
+	}
 }
 
 void UCompositeLayerShadowReflection::BeginDestroy()
 {
 	Super::BeginDestroy();
 
-	OnRemoved(GetWorld()); // Redundant remove call for safety
+	OnRemoved(GetTypedOuter<ACompositeActor>()); // Redundant remove call for safety
 }
 
 #if WITH_EDITOR
@@ -49,7 +62,14 @@ void UCompositeLayerShadowReflection::PostEditChangeProperty(FPropertyChangedEve
 	}
 	else if (PropertyName == GET_MEMBER_NAME_CHECKED(ThisClass, AutoConfigureActors))
 	{
-		UpdatePrimitiveVisibilityState();
+		if (AutoConfigureActors != ECompositeHiddenInSceneCaptureConfiguration::None)
+		{
+			const bool bHiddenInSceneCapture = (AutoConfigureActors == ECompositeHiddenInSceneCaptureConfiguration::Hidden);
+			const bool bAffectIndirectLightingWhileHidden = bHiddenInSceneCapture;
+			const bool bCastHiddenShadow = bHiddenInSceneCapture;
+
+			UpdatePrimitiveVisibilityState(bHiddenInSceneCapture, bAffectIndirectLightingWhileHidden, bCastHiddenShadow);
+		}
 	}
 }
 
@@ -66,6 +86,19 @@ void UCompositeLayerShadowReflection::PostEditUndo()
 }
 #endif
 
+void UCompositeLayerShadowReflection::OnEndOfFrameUpdate(UWorld* InWorld)
+{
+	/**
+	* In order to retain the ability to reference actors in sublevels, we prefer using weak pointer components on scene captures. However,
+	* primitive components may often be recreated (see RerunConstructionScripts(), FStaticMeshComponentRecreateRenderStateContext, dynamic
+	* text primitives, etc). To prevent issues with destroyed primitives disappearing in our scene captures, we re-register them every frame.
+	*/
+	if (CachedSceneCaptures[1].IsValid())
+	{
+		UpdateSceneCaptureComponents(*CachedSceneCaptures[1]);
+	}
+}
+
 bool UCompositeLayerShadowReflection::GetProxy(FTraversalContext& InContext, FSceneRenderingBulkObjectAllocator& InFrameAllocator, FCompositeCorePassProxy*& OutProxy) const
 {
 	using namespace UE::CompositeCore;
@@ -122,31 +155,16 @@ const TArray<TSoftObjectPtr<AActor>> UCompositeLayerShadowReflection::GetActors(
 
 void UCompositeLayerShadowReflection::SetActors(TArray<TSoftObjectPtr<AActor>> InActors)
 {
-	RestorePrimitiveVisibilityState();
-
 	Actors = MoveTemp(InActors);
 
-	// Update hidden components
-	TArray<TWeakObjectPtr<UPrimitiveComponent> > HiddenComponents;
-	HiddenComponents.Reserve(Actors.Num());
-
-	for (const TSoftObjectPtr<AActor>& SoftActor : Actors)
+	if (AutoConfigureActors != ECompositeHiddenInSceneCaptureConfiguration::None)
 	{
-		if (const AActor* Actor = SoftActor.Get())
-		{
-			for (UActorComponent* Component : Actor->GetComponents())
-			{
-				if (UPrimitiveComponent* PrimitiveComponent = Cast<UPrimitiveComponent>(Component))
-				{
-					HiddenComponents.Add(PrimitiveComponent);
-				}
-			}
-		}
-	}
+		const bool bHiddenInSceneCapture = (AutoConfigureActors == ECompositeHiddenInSceneCaptureConfiguration::Hidden);
+		const bool bAffectIndirectLightingWhileHidden = bHiddenInSceneCapture;
+		const bool bCastHiddenShadow = bHiddenInSceneCapture;
 
-	CachePrimitiveVisibilityState(HiddenComponents);
-
-	UpdatePrimitiveVisibilityState();
+		UpdatePrimitiveVisibilityState(bHiddenInSceneCapture, bAffectIndirectLightingWhileHidden, bCastHiddenShadow);
+	}
 
 	ACompositeActor* CompositeActor = GetTypedOuter<ACompositeActor>();
 	if (!IsValid(CompositeActor))
@@ -167,42 +185,15 @@ void UCompositeLayerShadowReflection::SetActors(TArray<TSoftObjectPtr<AActor>> I
 
 	// First, we empty hidden actors since our previous logic used them
 	CachedSceneCaptures[1]->HiddenActors.Empty();
-	CachedSceneCaptures[1]->HiddenComponents = HiddenComponents;
-}
-
-void UCompositeLayerShadowReflection::CachePrimitiveVisibilityState(TArrayView<TWeakObjectPtr<UPrimitiveComponent>> InPrimitives)
-{
-	CachedVisibilityInSceneCapture.Reset();
 
-	for (const TWeakObjectPtr<UPrimitiveComponent>& PrimitiveComponent : InPrimitives)
-	{
-		TStrongObjectPtr<UPrimitiveComponent> Primitive = PrimitiveComponent.Pin();
-		if (Primitive.IsValid())
-		{
-			FCompositeShadowReflectionPrimitiveState& State = CachedVisibilityInSceneCapture.Add(Primitive.Get());
-			State.bHiddenInSceneCapture = Primitive->bHiddenInSceneCapture;
-			State.bAffectIndirectLightingWhileHidden = Primitive->bAffectIndirectLightingWhileHidden;
-			State.bCastHiddenShadow = Primitive->bCastHiddenShadow;
-		}
-	}
+	UpdateSceneCaptureComponents(*CachedSceneCaptures[1]);
 }
 
-void UCompositeLayerShadowReflection::UpdatePrimitiveVisibilityState()
+void UCompositeLayerShadowReflection::UpdatePrimitiveVisibilityState(
+	bool bInHiddenInSceneCapture,
+	bool bInAffectIndirectLightingWhileHidden,
+	bool bInCastHiddenShadow)
 {
-	bool bHideActorsInSceneCapture = false;
-
-	switch (AutoConfigureActors)
-	{
-	case ECompositeHiddenInSceneCaptureConfiguration::None:
-		return;
-	case ECompositeHiddenInSceneCaptureConfiguration::Hidden:
-		bHideActorsInSceneCapture = true;
-		break;
-	case ECompositeHiddenInSceneCaptureConfiguration::Visible:
-	default:
-		break;
-	}
-
 	for (TSoftObjectPtr<AActor>& Actor : Actors)
 	{
 		if (Actor.IsValid())
@@ -211,33 +202,53 @@ void UCompositeLayerShadowReflection::UpdatePrimitiveVisibilityState()
 			{
 				if (UPrimitiveComponent* PrimitiveComponent = Cast<UPrimitiveComponent>(Component))
 				{
-					PrimitiveComponent->Modify();
-					PrimitiveComponent->SetHiddenInSceneCapture(bHideActorsInSceneCapture);
-					PrimitiveComponent->SetAffectIndirectLightingWhileHidden(bHideActorsInSceneCapture);
-					PrimitiveComponent->SetCastHiddenShadow(bHideActorsInSceneCapture);
+					bool bUpdatedComponent = false;
+
+					if (PrimitiveComponent->bHiddenInSceneCapture != bInHiddenInSceneCapture)
+					{
+						PrimitiveComponent->SetHiddenInSceneCapture(bInHiddenInSceneCapture);
+						bUpdatedComponent = true;
+					}
+
+					if (PrimitiveComponent->bAffectIndirectLightingWhileHidden != bInAffectIndirectLightingWhileHidden)
+					{
+						PrimitiveComponent->SetAffectIndirectLightingWhileHidden(bInAffectIndirectLightingWhileHidden);
+						bUpdatedComponent = true;
+					}
+
+					if (PrimitiveComponent->bCastHiddenShadow != bInCastHiddenShadow)
+					{
+						PrimitiveComponent->SetCastHiddenShadow(bInCastHiddenShadow);
+						bUpdatedComponent = true;
+					}
+
+					if (bUpdatedComponent)
+					{
+						PrimitiveComponent->Modify();
+					}
 				}
 			}
 		}
 	}
 }
 
-void UCompositeLayerShadowReflection::RestorePrimitiveVisibilityState()
+void UCompositeLayerShadowReflection::UpdateSceneCaptureComponents(USceneCaptureComponent2D& SceneCaptureComponent) const
 {
-	for (const TPair<TWeakObjectPtr<UPrimitiveComponent>, FCompositeShadowReflectionPrimitiveState>& Pair : CachedVisibilityInSceneCapture)
+	SceneCaptureComponent.HiddenComponents.Reset(Actors.Num());
+
+	for (const TSoftObjectPtr<AActor>& SoftActor : Actors)
 	{
-		TStrongObjectPtr< UPrimitiveComponent> Primitive = Pair.Key.Pin();
-		if (Primitive.IsValid())
+		if (const AActor* Actor = SoftActor.Get())
 		{
-			const FCompositeShadowReflectionPrimitiveState& State = Pair.Value;
-
-			Primitive->Modify();
-			Primitive->SetHiddenInSceneCapture(State.bHiddenInSceneCapture);
-			Primitive->SetAffectIndirectLightingWhileHidden(State.bAffectIndirectLightingWhileHidden);
-			Primitive->SetCastHiddenShadow(State.bCastHiddenShadow);
+			for (UActorComponent* Component : Actor->GetComponents())
+			{
+				if (UPrimitiveComponent* PrimitiveComponent = Cast<UPrimitiveComponent>(Component))
+				{
+					SceneCaptureComponent.HiddenComponents.Add(PrimitiveComponent);
+				}
+			}
 		}
 	}
-
-	CachedVisibilityInSceneCapture.Reset();
 }
 
 TStaticArray<TWeakObjectPtr<UCompositeShadowReflectionCatcherComponent>, 2> UCompositeLayerShadowReflection::FindOrCreateSceneCapturePair(ACompositeActor& InOuter) const
diff --git a/Engine/Plugins/Compositing/Composite/Source/Composite/Private/Layers/CompositeLayerSingleLightShadow.cpp b/Engine/Plugins/Compositing/Composite/Source/Composite/Private/Layers/CompositeLayerSingleLightShadow.cpp
--- a/Engine/Plugins/Compositing/Composite/Source/Composite/Private/Layers/CompositeLayerSingleLightShadow.cpp
+++ b/Engine/Plugins/Compositing/Composite/Source/Composite/Private/Layers/CompositeLayerSingleLightShadow.cpp
@@ -162,15 +162,22 @@ UCompositeLayerSingleLightShadow::UCompositeLayerSingleLightShadow(const FObject
 	: Super(ObjectInitializer)
 {
 	Operation = ECompositeCoreMergeOp::Multiply;
+
+	if (!HasAnyFlags(RF_ClassDefaultObject))
+	{
+		constexpr bool bEnabled = true;
+		RegisterEndOfFrameUpdate(bEnabled);
+	}
 }
 
-void UCompositeLayerSingleLightShadow::OnRemoved(const UWorld* World)
+void UCompositeLayerSingleLightShadow::OnRemoved(ACompositeActor* LastOwner)
 {
-	ACompositeActor* CompositeActor = GetTypedOuter<ACompositeActor>();
+	Super::OnRemoved(LastOwner);
+
 	// Note: We don't use IsValid(..) since destruction should proceed even if the actor is pending kill.
-	if (CompositeActor != nullptr)
+	if (LastOwner != nullptr)
 	{
-		CompositeActor->DestroySceneCaptures(this);
+		LastOwner->DestroySceneCaptures(this);
 	}
 }
 
@@ -178,7 +185,7 @@ void UCompositeLayerSingleLightShadow::BeginDestroy()
 {
 	Super::BeginDestroy();
 
-	OnRemoved(GetWorld()); // Redundant remove call for safety
+	OnRemoved(GetTypedOuter<ACompositeActor>()); // Redundant remove call for safety
 }
 
 #if WITH_EDITOR
@@ -202,6 +209,19 @@ void UCompositeLayerSingleLightShadow::PostEditUndo()
 }
 #endif
 
+void UCompositeLayerSingleLightShadow::OnEndOfFrameUpdate(UWorld* InWorld)
+{
+	/**
+	* In order to retain the ability to reference actors in sublevels, we prefer using weak pointer components on scene captures. However,
+	* primitive components may often be recreated (see RerunConstructionScripts(), FStaticMeshComponentRecreateRenderStateContext, dynamic
+	* text primitives, etc). To prevent issues with destroyed primitives disappearing in our scene captures, we re-register them every frame.
+	*/
+	if (CachedSceneDepthCapture.IsValid())
+	{
+		UpdateSceneCaptureComponents(*CachedSceneDepthCapture);
+	}
+}
+
 bool UCompositeLayerSingleLightShadow::GetProxy(FTraversalContext& InContext, FSceneRenderingBulkObjectAllocator& InFrameAllocator, FCompositeCorePassProxy*& OutProxy) const
 {
 	using namespace UE::CompositeCore;
@@ -299,34 +319,13 @@ void UCompositeLayerSingleLightShadow::SetShadowCastingActors(TArray<TSoftObject
 		CachedSceneDepthCapture = FindOrCreateSceneDepthCapture(*CompositeActor);
 	}
 	
-	// Early exit if the cached scene depth capture is still not valid
-	if (!CachedSceneDepthCapture.IsValid())
+	if (CachedSceneDepthCapture.IsValid())
 	{
-		return;
-	}
-
-	// First, we empty hidden actors since our previous logic used them
-	CachedSceneDepthCapture->HiddenActors.Empty();
-
-	// Collect & set hidden primitive components
-	TArray<TWeakObjectPtr<UPrimitiveComponent> > HiddenComponents;
-	HiddenComponents.Reserve(ShadowCastingActors.Num());
+		// First, we empty hidden actors since our previous logic used them
+		CachedSceneDepthCapture->HiddenActors.Empty();
 
-	for (const TSoftObjectPtr<AActor>& SoftActor : ShadowCastingActors)
-	{
-		if (const AActor* Actor = SoftActor.Get())
-		{
-			for (UActorComponent* Component : Actor->GetComponents())
-			{
-				if (UPrimitiveComponent* PrimitiveComponent = Cast<UPrimitiveComponent>(Component))
-				{
-					HiddenComponents.Add(PrimitiveComponent);
-				}
-			}
-		}
+		UpdateSceneCaptureComponents(*CachedSceneDepthCapture);
 	}
-
-	CachedSceneDepthCapture->HiddenComponents = HiddenComponents;
 }
 
 UCompositeSceneCapture2DComponent* UCompositeLayerSingleLightShadow::FindOrCreateSceneDepthCapture(ACompositeActor& InOuter) const
@@ -364,6 +363,26 @@ UCompositeSceneCapture2DComponent* UCompositeLayerSingleLightShadow::FindOrCreat
 	return ShadowMapCapture;
 }
 
+void UCompositeLayerSingleLightShadow::UpdateSceneCaptureComponents(USceneCaptureComponent2D& SceneCaptureComponent) const
+{
+	// Collect & set hidden primitive components
+	SceneCaptureComponent.HiddenComponents.Reset(ShadowCastingActors.Num());
+
+	for (const TSoftObjectPtr<AActor>& SoftActor : ShadowCastingActors)
+	{
+		if (const AActor* Actor = SoftActor.Get())
+		{
+			for (UActorComponent* Component : Actor->GetComponents())
+			{
+				if (UPrimitiveComponent* PrimitiveComponent = Cast<UPrimitiveComponent>(Component))
+				{
+					SceneCaptureComponent.HiddenComponents.Add(PrimitiveComponent);
+				}
+			}
+		}
+	}
+}
+
 void UCompositeLayerSingleLightShadow::GetShadowMatrices(UCompositeSceneCapture2DComponent& ShadowMapCapture, FMatrix44f& OutShadowMatrix, FVector4f& OutShadowInvDeviceZToWorldZ) const
 {
 	FMinimalViewInfo ViewInfo;
diff --git a/Engine/Plugins/Compositing/Composite/Source/Composite/Public/Layers/CompositeLayerBase.h b/Engine/Plugins/Compositing/Composite/Source/Composite/Public/Layers/CompositeLayerBase.h
--- a/Engine/Plugins/Compositing/Composite/Source/Composite/Public/Layers/CompositeLayerBase.h
+++ b/Engine/Plugins/Compositing/Composite/Source/Composite/Public/Layers/CompositeLayerBase.h
@@ -64,6 +64,12 @@ public:
 	UE_API virtual void PostEditChangeProperty(FPropertyChangedEvent& PropertyChangedEvent) override;
 #endif
 
+	/**
+	 * End-of-frame update callback to inherit.
+	 * Called when enabled by RegisterEndOfFrameUpdate.
+	*/
+	UE_API virtual void OnEndOfFrameUpdate(UWorld* InWorld) {}
+
 	/**
 	* Called when composite actor rendering state changes.
 	* 
@@ -74,9 +80,9 @@ public:
 	/**
 	* Called when the layer is removed or composite actor is destroyed.
 	*
-	* @param World - World from which the layer is removed. Ignore if null.
+	* @param LastOwner - Last composite actor owner, from which the layer was removed. Ignore if null.
 	*/
-	UE_API virtual void OnRemoved(const UWorld* World) {}
+	UE_API virtual void OnRemoved(ACompositeActor* LastOwner);
 
 	/** Getter function to override, returning pass layer proxies to be passed to the render thread. (Proxy objects should be allocated from the provided allocator.) */
 	UE_API virtual bool GetProxy(FTraversalContext& InContext, FSceneRenderingBulkObjectAllocator& InFrameAllocator, FCompositeCorePassProxy*& OutProxy) const { return false; };
@@ -89,6 +95,10 @@ public:
 	UFUNCTION(BlueprintSetter)
 	UE_API virtual void SetEnabled(bool bInEnabled);
 
+protected:
+	/** Activate the end-of-frame update callback. */
+	UE_API void RegisterEndOfFrameUpdate(bool bInEnabled);
+
 public:
 	/** Whether or not the pass is solo. */
 	UPROPERTY()
@@ -123,6 +133,10 @@ protected:
 
 	/** Convenience function that returns the default secondary input depending on the traversal context. */
 	UE::CompositeCore::FPassInputDecl GetDefaultSecondInput(const FTraversalContext& InContext) const;
+
+private:
+	/** Post-viewport update callback. */
+	FDelegateHandle OnWorldPreSendAllEndOfFrameUpdatesHandle;
 };
 
 #undef UE_API
diff --git a/Engine/Plugins/Compositing/Composite/Source/Composite/Public/Layers/CompositeLayerPlate.h b/Engine/Plugins/Compositing/Composite/Source/Composite/Public/Layers/CompositeLayerPlate.h
--- a/Engine/Plugins/Compositing/Composite/Source/Composite/Public/Layers/CompositeLayerPlate.h
+++ b/Engine/Plugins/Compositing/Composite/Source/Composite/Public/Layers/CompositeLayerPlate.h
@@ -3,7 +3,6 @@
 #pragma once
 
 #include "CompositeLayerBase.h"
-#include "Passes/CompositePassOpenColorIO.h"
 #include "Tickable.h"
 
 #include "CompositeLayerPlate.generated.h"
@@ -44,7 +43,7 @@ public:
 	/** Destructor. */
 	UE_API virtual ~UCompositeLayerPlate();
 
-	UE_API virtual void OnRemoved(const UWorld* World) override;
+	UE_API virtual void OnRemoved(ACompositeActor* LastOwner) override;
 
 	UE_API virtual void OnRenderingStateChange(bool bApply) override;
 
diff --git a/Engine/Plugins/Compositing/Composite/Source/Composite/Public/Layers/CompositeLayerSceneCapture.h b/Engine/Plugins/Compositing/Composite/Source/Composite/Public/Layers/CompositeLayerSceneCapture.h
--- a/Engine/Plugins/Compositing/Composite/Source/Composite/Public/Layers/CompositeLayerSceneCapture.h
+++ b/Engine/Plugins/Compositing/Composite/Source/Composite/Public/Layers/CompositeLayerSceneCapture.h
@@ -10,6 +10,7 @@
 
 class AActor;
 class UPrimitiveComponent;
+class USceneCaptureComponent2D;
 
 /**
 * Layer for scene capture (or custom render pass) render integration.
@@ -26,7 +27,7 @@ public:
 	/** Destructor. */
 	UE_API virtual ~UCompositeLayerSceneCapture();
 	
-	UE_API virtual void OnRemoved(const UWorld* World) override;
+	UE_API virtual void OnRemoved(ACompositeActor* LastOwner) override;
 
 	//~ Begin UObject Interface
 	UE_API virtual void BeginDestroy() override;
@@ -38,6 +39,9 @@ public:
 #endif
 	//~ End UObject Interface
 
+	/** End-of-frame update callback. */
+	UE_API virtual void OnEndOfFrameUpdate(UWorld* InWorld) override;
+
 	/** Getter function to override, returning pass layer proxies to be passed to the render thread. (Proxy objects should be allocated from the provided allocator.) */
 	UE_API virtual bool GetProxy(FTraversalContext& InContext, FSceneRenderingBulkObjectAllocator& InFrameAllocator, FCompositeCorePassProxy*& OutProxy) const override;
 
@@ -70,17 +74,15 @@ public:
 	TArray<TObjectPtr<UCompositePassBase>> LayerPasses;
 
 private:
+
 	/** Convenience function to access the actor-managed scene capture component. */
 	class UCompositeSceneCapture2DComponent* GetSceneCaptureComponent() const;
 
-	/** Cache scene capture visibility on registered mesh actors primitives. */
-	void CachePrimitiveVisibilityState(TArrayView<TWeakObjectPtr<UPrimitiveComponent>> InPrimitives);
-
 	/** Update scene capture visibility on registered mesh actors primitives. */
-	void UpdatePrimitiveVisibilityState();
-	
-	/** Restore scene capture visibility on registered mesh actors primitives. */
-	void RestorePrimitiveVisibilityState();
+	void UpdatePrimitiveVisibilityState(bool bInVisibleInSceneCaptureOnly) const;
+
+	/** Update scene capture show only components. */
+	void UpdateSceneCaptureComponents(USceneCaptureComponent2D& SceneCaptureComponent) const;
 
 	/**
 	 * List of actors to be rendered from the scene capture.
@@ -90,10 +92,6 @@ private:
 	UPROPERTY(EditAnywhere, BlueprintGetter = GetActors, BlueprintSetter = SetActors, Category = "", meta = (AllowPrivateAccess = true))
 	TArray<TSoftObjectPtr<AActor>> Actors;
 
-	/** Cached visibility in scene capture state to restore initial values. */
-	UPROPERTY()
-	TMap<TWeakObjectPtr<UPrimitiveComponent>, bool> CachedVisibilityInSceneCapture;
-
 	/** Visibility setting applied to registered actors. By default, registered meshes will only be visible in scene captures and no longer in the main render. */
 	UPROPERTY(EditAnywhere, BlueprintGetter = IsVisibleInSceneCaptureOnly, BlueprintSetter = SetVisibleInSceneCaptureOnly, Category = "", meta = (AllowPrivateAccess = true))
 	bool bVisibleInSceneCaptureOnly;
diff --git a/Engine/Plugins/Compositing/Composite/Source/Composite/Public/Layers/CompositeLayerShadowReflection.h b/Engine/Plugins/Compositing/Composite/Source/Composite/Public/Layers/CompositeLayerShadowReflection.h
--- a/Engine/Plugins/Compositing/Composite/Source/Composite/Public/Layers/CompositeLayerShadowReflection.h
+++ b/Engine/Plugins/Compositing/Composite/Source/Composite/Public/Layers/CompositeLayerShadowReflection.h
@@ -25,26 +25,6 @@ enum class ECompositeHiddenInSceneCaptureConfiguration : uint8
 	Hidden UMETA(ToolTip = "The following properties are set to true: bHiddenInSceneCapture, bAffectIndirectLightingWhileHidden & bCastHiddenShadow."),
 };
 
-/** Primitive visibility settings cached state. */
-USTRUCT()
-struct FCompositeShadowReflectionPrimitiveState
-{
-public:
-	GENERATED_USTRUCT_BODY()
-
-	/** Hidden in scene capture. */
-	UPROPERTY()
-	bool bHiddenInSceneCapture = false;
-	
-	/** Affect indirect lighting while hidden. */
-	UPROPERTY()
-	bool bAffectIndirectLightingWhileHidden = false;
-	
-	/** Cast hidden shadow. */
-	UPROPERTY()
-	bool bCastHiddenShadow = false;
-};
-
 /**
 * Layer for shadow & reflection catching with two additional scene capture renders.
 * Primitives are rendered in one scene capture but not the other. The division of both creates a multiplicative matte that can be applied onto the plate.
@@ -61,7 +41,7 @@ public:
 	/** Destructor. */
 	UE_API virtual ~UCompositeLayerShadowReflection();
 
-	UE_API virtual void OnRemoved(const UWorld* World) override;
+	UE_API virtual void OnRemoved(ACompositeActor* LastOwner) override;
 
 	//~ Begin UObject Interface
 	UE_API virtual void BeginDestroy() override;
@@ -73,6 +53,9 @@ public:
 #endif
 	//~ End UObject Interface
 
+	/** End-of-frame update callback. */
+	UE_API virtual void OnEndOfFrameUpdate(UWorld* InWorld) override;
+
 	/** Getter function to override, returning pass layer proxies to be passed to the render thread. (Proxy objects should be allocated from the provided allocator.) */
 	UE_API virtual bool GetProxy(FTraversalContext& InContext, FSceneRenderingBulkObjectAllocator& InFrameAllocator, FCompositeCorePassProxy*& OutProxy) const override;
 
@@ -95,14 +78,11 @@ public:
 	ECompositeHiddenInSceneCaptureConfiguration AutoConfigureActors;
 
 private:
-	/** Cache scene capture visibility on registered mesh actors primitives. */
-	void CachePrimitiveVisibilityState(TArrayView<TWeakObjectPtr<UPrimitiveComponent>> InPrimitives);
-
 	/** Update scene capture visibility on registered mesh actors primitives. */
-	void UpdatePrimitiveVisibilityState();
+	void UpdatePrimitiveVisibilityState(bool bInHiddenInSceneCapture, bool bInAffectIndirectLightingWhileHidden, bool bInCastHiddenShadow);
 
-	/** Restore scene capture visibility on registered mesh actors primitives. */
-	void RestorePrimitiveVisibilityState();
+	/** Update scene capture hidden components. */
+	void UpdateSceneCaptureComponents(USceneCaptureComponent2D& SceneCaptureComponent) const;
 
 	/** Find or create scene captures with CG and without CG. */
 	TStaticArray<TWeakObjectPtr<UCompositeShadowReflectionCatcherComponent>, 2> FindOrCreateSceneCapturePair(ACompositeActor& InOuter) const;
@@ -112,10 +92,6 @@ private:
 	UPROPERTY(EditAnywhere, BlueprintGetter = GetActors, BlueprintSetter = SetActors, Category = "")
 	TArray<TSoftObjectPtr<AActor>> Actors;
 
-	/** Cached visibility in scene capture state to restore initial values. */
-	UPROPERTY()
-	TMap<TWeakObjectPtr<UPrimitiveComponent>, FCompositeShadowReflectionPrimitiveState> CachedVisibilityInSceneCapture;
-
 	/** Cached pointer to the scene captures with CG and without CG. */
 	mutable TStaticArray<TWeakObjectPtr<UCompositeShadowReflectionCatcherComponent>, 2> CachedSceneCaptures;
 
diff --git a/Engine/Plugins/Compositing/Composite/Source/Composite/Public/Layers/CompositeLayerSingleLightShadow.h b/Engine/Plugins/Compositing/Composite/Source/Composite/Public/Layers/CompositeLayerSingleLightShadow.h
--- a/Engine/Plugins/Compositing/Composite/Source/Composite/Public/Layers/CompositeLayerSingleLightShadow.h
+++ b/Engine/Plugins/Compositing/Composite/Source/Composite/Public/Layers/CompositeLayerSingleLightShadow.h
@@ -37,7 +37,7 @@ public:
 	/** Destructor. */
 	UE_API virtual ~UCompositeLayerSingleLightShadow() = default;
 
-	UE_API virtual void OnRemoved(const UWorld* World) override;
+	UE_API virtual void OnRemoved(ACompositeActor* LastOwner) override;
 
 	//~ Begin UObject Interface
 	UE_API virtual void BeginDestroy() override;
@@ -48,6 +48,9 @@ public:
 #endif
 	//~ End UObject Interface
 
+	/** End-of-frame update callback. */
+	UE_API virtual void OnEndOfFrameUpdate(UWorld* InWorld) override;
+
 	/** Getter function to override, returning pass layer proxies to be passed to the render thread. (Proxy objects should be allocated from the provided allocator.) */
 	UE_API virtual bool GetProxy(FTraversalContext& InContext, FSceneRenderingBulkObjectAllocator& InFrameAllocator, FCompositeCorePassProxy*& OutProxy) const override;
 
@@ -99,6 +102,9 @@ private:
 	/** Find or create shadow map capture component, correctly configured. */
 	UCompositeSceneCapture2DComponent* FindOrCreateShadowMapCapture(ACompositeActor& InOuter) const;
 
+	/** Update scene capture hidden components. */
+	void UpdateSceneCaptureComponents(USceneCaptureComponent2D& SceneCaptureComponent) const;
+
 	/** Convenience function to calculate shadow matrices. */
 	void GetShadowMatrices(UCompositeSceneCapture2DComponent& ShadowMapCapture, FMatrix44f& OutShadowMatrix, FVector4f& OutShadowInvDeviceZToWorldZ) const;
 
