Composure: Additional fix to CL 47600697.

Correctly restore default primitive settings when unregistered from layers, either from property edits, direct set actors call or undo/redo.

diff --git a/Engine/Plugins/Compositing/Composite/Source/Composite/Private/Layers/CompositeLayerSceneCapture.cpp b/Engine/Plugins/Compositing/Composite/Source/Composite/Private/Layers/CompositeLayerSceneCapture.cpp
--- a/Engine/Plugins/Compositing/Composite/Source/Composite/Private/Layers/CompositeLayerSceneCapture.cpp
+++ b/Engine/Plugins/Compositing/Composite/Source/Composite/Private/Layers/CompositeLayerSceneCapture.cpp
@@ -35,9 +35,7 @@ void UCompositeLayerSceneCapture::OnRemoved(ACompositeActor* LastOwner)
 		LastOwner->DestroySceneCaptures(this);
 	}
 
-	// Restore to class default
-	UPrimitiveComponent* CDO = UPrimitiveComponent::StaticClass()->GetDefaultObject<UPrimitiveComponent>();
-	UpdatePrimitiveVisibilityState(CDO->bVisibleInSceneCaptureOnly);
+	RestorePrimitiveVisibilityState();
 }
 
 void UCompositeLayerSceneCapture::BeginDestroy()
@@ -48,6 +46,23 @@ void UCompositeLayerSceneCapture::BeginDestroy()
 }
 
 #if WITH_EDITOR
+void UCompositeLayerSceneCapture::PreEditChange(FProperty* PropertyAboutToChange)
+{
+	Super::PreEditChange(PropertyAboutToChange);
+
+	if (!PropertyAboutToChange)
+	{
+		return;
+	}
+
+	const FName PropertyName = PropertyAboutToChange->GetFName();
+
+	if (PropertyName == GET_MEMBER_NAME_CHECKED(ThisClass, Actors))
+	{
+		RestorePrimitiveVisibilityState();
+	}
+}
+
 void UCompositeLayerSceneCapture::PostEditChangeProperty(FPropertyChangedEvent& PropertyChangedEvent)
 {
 	Super::PostEditChangeProperty(PropertyChangedEvent);
@@ -71,6 +86,8 @@ void UCompositeLayerSceneCapture::PostEditChangeProperty(FPropertyChangedEvent&
 void UCompositeLayerSceneCapture::PreEditUndo()
 {
 	Super::PreEditUndo();
+
+	RestorePrimitiveVisibilityState();
 }
 
 void UCompositeLayerSceneCapture::PostEditUndo()
@@ -142,6 +159,8 @@ const TArray<TSoftObjectPtr<AActor>> UCompositeLayerSceneCapture::GetActors() co
 
 void UCompositeLayerSceneCapture::SetActors(TArray<TSoftObjectPtr<AActor>> InActors)
 {
+	RestorePrimitiveVisibilityState();
+
 	Actors = MoveTemp(InActors);
 
 	UpdatePrimitiveVisibilityState(bVisibleInSceneCaptureOnly);
@@ -225,6 +244,14 @@ void UCompositeLayerSceneCapture::UpdatePrimitiveVisibilityState(bool bInVisible
 	}
 }
 
+void UCompositeLayerSceneCapture::RestorePrimitiveVisibilityState() const
+{
+	// First, we restore current actors to class default, in case they have been deleted.
+	UPrimitiveComponent* CDO = UPrimitiveComponent::StaticClass()->GetDefaultObject<UPrimitiveComponent>();
+
+	UpdatePrimitiveVisibilityState(CDO->bVisibleInSceneCaptureOnly);
+}
+
 void UCompositeLayerSceneCapture::UpdateSceneCaptureComponents(USceneCaptureComponent2D& SceneCaptureComponent) const
 {
 	SceneCaptureComponent.ShowOnlyComponents.Reset(Actors.Num());
diff --git a/Engine/Plugins/Compositing/Composite/Source/Composite/Private/Layers/CompositeLayerShadowReflection.cpp b/Engine/Plugins/Compositing/Composite/Source/Composite/Private/Layers/CompositeLayerShadowReflection.cpp
--- a/Engine/Plugins/Compositing/Composite/Source/Composite/Private/Layers/CompositeLayerShadowReflection.cpp
+++ b/Engine/Plugins/Compositing/Composite/Source/Composite/Private/Layers/CompositeLayerShadowReflection.cpp
@@ -35,10 +35,7 @@ void UCompositeLayerShadowReflection::OnRemoved(ACompositeActor* LastOwner)
 
 	if (AutoConfigureActors != ECompositeHiddenInSceneCaptureConfiguration::None)
 	{
-		// Restore to class default
-		UPrimitiveComponent* CDO = UPrimitiveComponent::StaticClass()->GetDefaultObject<UPrimitiveComponent>();
-
-		UpdatePrimitiveVisibilityState(CDO->bHiddenInSceneCapture, CDO->bAffectIndirectLightingWhileHidden, CDO->bCastHiddenShadow);
+		RestorePrimitiveVisibilityState();
 	}
 }
 
@@ -50,6 +47,26 @@ void UCompositeLayerShadowReflection::BeginDestroy()
 }
 
 #if WITH_EDITOR
+void UCompositeLayerShadowReflection::PreEditChange(FProperty* PropertyAboutToChange)
+{
+	Super::PreEditChange(PropertyAboutToChange);
+
+	if (!PropertyAboutToChange)
+	{
+		return;
+	}
+
+	const FName PropertyName = PropertyAboutToChange->GetFName();
+
+	if (PropertyName == GET_MEMBER_NAME_CHECKED(ThisClass, Actors))
+	{
+		if (AutoConfigureActors != ECompositeHiddenInSceneCaptureConfiguration::None)
+		{
+			RestorePrimitiveVisibilityState();
+		}
+	}
+}
+
 void UCompositeLayerShadowReflection::PostEditChangeProperty(FPropertyChangedEvent& PropertyChangedEvent)
 {
 	Super::PostEditChangeProperty(PropertyChangedEvent);
@@ -76,6 +93,11 @@ void UCompositeLayerShadowReflection::PostEditChangeProperty(FPropertyChangedEve
 void UCompositeLayerShadowReflection::PreEditUndo()
 {
 	Super::PreEditUndo();
+
+	if (AutoConfigureActors != ECompositeHiddenInSceneCaptureConfiguration::None)
+	{
+		RestorePrimitiveVisibilityState();
+	}
 }
 
 void UCompositeLayerShadowReflection::PostEditUndo()
@@ -155,6 +177,11 @@ const TArray<TSoftObjectPtr<AActor>> UCompositeLayerShadowReflection::GetActors(
 
 void UCompositeLayerShadowReflection::SetActors(TArray<TSoftObjectPtr<AActor>> InActors)
 {
+	if (AutoConfigureActors != ECompositeHiddenInSceneCaptureConfiguration::None)
+	{
+		RestorePrimitiveVisibilityState();
+	}
+
 	Actors = MoveTemp(InActors);
 
 	if (AutoConfigureActors != ECompositeHiddenInSceneCaptureConfiguration::None)
@@ -192,9 +219,9 @@ void UCompositeLayerShadowReflection::SetActors(TArray<TSoftObjectPtr<AActor>> I
 void UCompositeLayerShadowReflection::UpdatePrimitiveVisibilityState(
 	bool bInHiddenInSceneCapture,
 	bool bInAffectIndirectLightingWhileHidden,
-	bool bInCastHiddenShadow)
+	bool bInCastHiddenShadow) const
 {
-	for (TSoftObjectPtr<AActor>& Actor : Actors)
+	for (const TSoftObjectPtr<AActor>& Actor : Actors)
 	{
 		if (Actor.IsValid())
 		{
@@ -232,6 +259,14 @@ void UCompositeLayerShadowReflection::UpdatePrimitiveVisibilityState(
 	}
 }
 
+void UCompositeLayerShadowReflection::RestorePrimitiveVisibilityState() const
+{
+	// Restore to class default
+	UPrimitiveComponent* CDO = UPrimitiveComponent::StaticClass()->GetDefaultObject<UPrimitiveComponent>();
+
+	UpdatePrimitiveVisibilityState(CDO->bHiddenInSceneCapture, CDO->bAffectIndirectLightingWhileHidden, CDO->bCastHiddenShadow);
+}
+
 void UCompositeLayerShadowReflection::UpdateSceneCaptureComponents(USceneCaptureComponent2D& SceneCaptureComponent) const
 {
 	SceneCaptureComponent.HiddenComponents.Reset(Actors.Num());
diff --git a/Engine/Plugins/Compositing/Composite/Source/Composite/Public/Layers/CompositeLayerSceneCapture.h b/Engine/Plugins/Compositing/Composite/Source/Composite/Public/Layers/CompositeLayerSceneCapture.h
--- a/Engine/Plugins/Compositing/Composite/Source/Composite/Public/Layers/CompositeLayerSceneCapture.h
+++ b/Engine/Plugins/Compositing/Composite/Source/Composite/Public/Layers/CompositeLayerSceneCapture.h
@@ -33,6 +33,7 @@ public:
 	UE_API virtual void BeginDestroy() override;
 
 #if WITH_EDITOR
+	UE_API virtual void PreEditChange(FProperty* PropertyAboutToChange) override;
 	UE_API virtual void PostEditChangeProperty(FPropertyChangedEvent& PropertyChangedEvent) override;
 	UE_API virtual void PreEditUndo() override;
 	UE_API virtual void PostEditUndo() override;
@@ -81,6 +82,9 @@ private:
 	/** Update scene capture visibility on registered mesh actors primitives. */
 	void UpdatePrimitiveVisibilityState(bool bInVisibleInSceneCaptureOnly) const;
 
+	/** Restore scene capture visibility on registered mesh actors primitives. */
+	void RestorePrimitiveVisibilityState()  const;
+
 	/** Update scene capture show only components. */
 	void UpdateSceneCaptureComponents(USceneCaptureComponent2D& SceneCaptureComponent) const;
 
diff --git a/Engine/Plugins/Compositing/Composite/Source/Composite/Public/Layers/CompositeLayerShadowReflection.h b/Engine/Plugins/Compositing/Composite/Source/Composite/Public/Layers/CompositeLayerShadowReflection.h
--- a/Engine/Plugins/Compositing/Composite/Source/Composite/Public/Layers/CompositeLayerShadowReflection.h
+++ b/Engine/Plugins/Compositing/Composite/Source/Composite/Public/Layers/CompositeLayerShadowReflection.h
@@ -47,6 +47,7 @@ public:
 	UE_API virtual void BeginDestroy() override;
 
 #if WITH_EDITOR
+	UE_API virtual void PreEditChange(FProperty* PropertyAboutToChange) override;
 	UE_API virtual void PostEditChangeProperty(FPropertyChangedEvent& PropertyChangedEvent) override;
 	UE_API virtual void PreEditUndo() override;
 	UE_API virtual void PostEditUndo() override;
@@ -79,7 +80,10 @@ public:
 
 private:
 	/** Update scene capture visibility on registered mesh actors primitives. */
-	void UpdatePrimitiveVisibilityState(bool bInHiddenInSceneCapture, bool bInAffectIndirectLightingWhileHidden, bool bInCastHiddenShadow);
+	void UpdatePrimitiveVisibilityState(bool bInHiddenInSceneCapture, bool bInAffectIndirectLightingWhileHidden, bool bInCastHiddenShadow) const;
+
+	/** Restore scene capture visibility on registered mesh actors primitives. */
+	void RestorePrimitiveVisibilityState() const;
 
 	/** Update scene capture hidden components. */
 	void UpdateSceneCaptureComponents(USceneCaptureComponent2D& SceneCaptureComponent) const;
