Composure: Preliminary support for (media texture) depth.

* New Composite Depth Mesh actor hosts a grid displaced in WPO by de-projecting the depth texture (no custom vertex factory yet).
* Support for scene depth or device Z via material instance static switch, with additional scale & offset parameters for scene depth (to be formalized & exposed later).

Note: poor performance & visual temporal artifacts with Lumen (especially reflections) hence the unlit material for now.

#virtualized
#Synced-CL 48124759

diff --git a/Engine/Plugins/Compositing/Composite/Composite.uplugin b/Engine/Plugins/Compositing/Composite/Composite.uplugin
--- a/Engine/Plugins/Compositing/Composite/Composite.uplugin
+++ b/Engine/Plugins/Compositing/Composite/Composite.uplugin
@@ -51,6 +51,11 @@
 		{
 			"Name": "ConcertSyncCore",
 			"Enabled": true
-		}
+		},
+		{
+			"Name": "MeshModelingToolset",
+			"TargetAllowList": [ "Editor" ],
+			"Enabled": true
+		},
 	]
 }
\ No newline at end of file
diff --git a/Engine/Plugins/Compositing/Composite/Config/DefaultComposite.ini b/Engine/Plugins/Compositing/Composite/Config/DefaultComposite.ini
--- a/Engine/Plugins/Compositing/Composite/Config/DefaultComposite.ini
+++ b/Engine/Plugins/Compositing/Composite/Config/DefaultComposite.ini
@@ -37,3 +37,6 @@
 +EnumRedirects=(OldName="/Script/Composite.CompositeDeltaKeyerScreenType",NewName="/Script/Composite.CompositeColorKeyerScreenType")
 +ClassRedirects=(OldName="/Script/Composite.CompositePassDeltaKeyer", NewName="/Script/Composite.CompositePassColorKeyer")
 
+; 5.8
+
++PropertyRedirects=(OldName="/Script/Composite.CompositeViewProjectionComponent.MatrixParameterName", NewName="/Script/Composite.CompositeViewProjectionComponent.ViewProjectionMatrixParameter")
diff --git a/Engine/Plugins/Compositing/Composite/Source/Composite/Composite.Build.cs b/Engine/Plugins/Compositing/Composite/Source/Composite/Composite.Build.cs
--- a/Engine/Plugins/Compositing/Composite/Source/Composite/Composite.Build.cs
+++ b/Engine/Plugins/Compositing/Composite/Source/Composite/Composite.Build.cs
@@ -44,6 +44,8 @@ public class Composite : ModuleRules
 					"UnrealEd",
 					"Slate",
 					"SlateCore",
+					"GeometryCore",
+					"ModelingComponentsEditorOnly",
 				}
 			);
 		}
diff --git a/Engine/Plugins/Compositing/Composite/Source/Composite/Private/Components/CompositeDepthMeshComponent.cpp b/Engine/Plugins/Compositing/Composite/Source/Composite/Private/Components/CompositeDepthMeshComponent.cpp
new file mode 100644
--- /dev/null
+++ b/Engine/Plugins/Compositing/Composite/Source/Composite/Private/Components/CompositeDepthMeshComponent.cpp
@@ -0,0 +1,205 @@
+// Copyright Epic Games, Inc. All Rights Reserved.
+
+#include "Components/CompositeDepthMeshComponent.h"
+
+#include "Engine/StaticMesh.h"
+#include "Engine/Texture.h"
+#include "Materials/Material.h"
+#include "Materials/MaterialInstanceDynamic.h"
+#include "UObject/ConstructorHelpers.h"
+
+#if WITH_EDITOR
+#include "AssetRegistry/AssetRegistryModule.h"
+#include "AssetUtils/CreateStaticMeshUtil.h"
+#include "DynamicMesh/DynamicMesh3.h"
+#include "Generators/RectangleMeshGenerator.h"
+#include "Misc/TransactionObjectEvent.h"
+#endif
+
+#define LOCTEXT_NAMESPACE "CompositingUtilities"
+
+namespace UE
+{
+	namespace Composite
+	{
+		namespace Private
+		{
+			static FLazyName CompositeDepthTextureName(TEXT("CompositeDepthTexture"));
+
+#if WITH_EDITOR
+			FString GetAssetPath(UCompositeDepthMeshComponent& InCompositeDepthMesh)
+			{
+				FString Res_X = FString::FromInt(InCompositeDepthMesh.GridResolution.X);
+				FString Res_Y = FString::FromInt(InCompositeDepthMesh.GridResolution.Y);
+				FString AssetPath = FString::Printf(TEXT("/Game/_Composure/Generated/SM_Grid_%s_%s"), *Res_X, *Res_Y);
+
+				return AssetPath;
+			}
+
+			UStaticMesh* CreateStaticMeshAsset(UE::Geometry::FDynamicMesh3* Mesh, const FString& AssetPath)
+			{
+				UStaticMesh* NewStaticMesh = nullptr;
+
+				UE::AssetUtils::FStaticMeshAssetOptions AssetOptions;
+
+				AssetOptions.NewAssetPath = AssetPath;
+				AssetOptions.NumSourceModels = 1;
+
+				AssetOptions.bEnableRecomputeNormals = false;
+				AssetOptions.bEnableRecomputeTangents = false;
+				AssetOptions.bGenerateNaniteEnabledMesh = false;
+				AssetOptions.NaniteSettings.FallbackPercentTriangles = 1.0;
+
+				AssetOptions.bCreatePhysicsBody = true;
+				AssetOptions.CollisionType = ECollisionTraceFlag::CTF_UseComplexAsSimple;
+
+				AssetOptions.SourceMeshes.DynamicMeshes.Add(Mesh);
+
+				UE::AssetUtils::FStaticMeshResults ResultData;
+				UE::AssetUtils::ECreateStaticMeshResult AssetResult = UE::AssetUtils::CreateStaticMeshAsset(AssetOptions, ResultData);
+
+				if (AssetResult == UE::AssetUtils::ECreateStaticMeshResult::Ok)
+				{
+					NewStaticMesh = ResultData.StaticMesh;
+					NewStaticMesh->MarkPackageDirty();
+					FAssetRegistryModule::AssetCreated(NewStaticMesh);
+				}
+
+				return NewStaticMesh;
+			}
+#endif
+		}
+	}
+}
+
+UCompositeDepthMeshComponent::UCompositeDepthMeshComponent(const FObjectInitializer& ObjectInitializer)
+	: UStaticMeshComponent(ObjectInitializer)
+	, GridResolution{1920/2, 1080/2}
+	, DepthTexture{}
+{
+	struct FConstructorStatics
+	{
+		ConstructorHelpers::FObjectFinder<UStaticMesh> DefaultDepthMesh;
+		ConstructorHelpers::FObjectFinder<UMaterialInterface> DefaultDepthMaterial;
+		FConstructorStatics()
+			: DefaultDepthMesh(TEXT("/Composite/Meshes/SM_Grid_960_540.SM_Grid_960_540"))
+			, DefaultDepthMaterial(TEXT("/Composite/Materials/M_CompositeDepthMesh_Unlit_Masked.M_CompositeDepthMesh_Unlit_Masked"))
+		{
+		}
+	};
+
+	static FConstructorStatics ConstructorStatics;
+	if (ensure(ConstructorStatics.DefaultDepthMesh.Object != nullptr))
+	{
+		SetStaticMesh(ConstructorStatics.DefaultDepthMesh.Object);
+	}
+
+	if (ensure(ConstructorStatics.DefaultDepthMaterial.Object != nullptr))
+	{
+		SetMaterial(0, ConstructorStatics.DefaultDepthMaterial.Object);
+	}
+
+	SetEnableGravity(false);
+	SetGenerateOverlapEvents(false);
+	SetCollisionEnabled(ECollisionEnabled::NoCollision);
+	SetCastShadow(false);
+	SetEvaluateWorldPositionOffset(true);
+	
+	// Increase bounds to prevent flickering with WPO material
+	BoundsScale = 1000.0f;
+}
+
+void UCompositeDepthMeshComponent::PostLoad()
+{
+	Super::PostLoad();
+
+	if (!HasAnyFlags(RF_ClassDefaultObject))
+	{
+		UpdateMaterial();
+	}
+}
+
+void UCompositeDepthMeshComponent::UpdateMaterial()
+{
+	UMaterialInterface* MaterialInterface = GetMaterial(0);
+	if (!IsValid(MaterialInterface))
+	{
+		return;
+	}
+
+	TArray<FMaterialParameterInfo> ParameterInfo;
+	TArray<FGuid> ParameterIds;
+	MaterialInterface->GetAllTextureParameterInfo(ParameterInfo, ParameterIds);
+
+	// Avoid conversion to MID if parameter doesn't exist
+	const bool bHasTextureParameter = ParameterInfo.ContainsByPredicate([](const FMaterialParameterInfo& ParamInfo)
+		{
+			return ParamInfo.Name.IsEqual(UE::Composite::Private::CompositeDepthTextureName);
+		}
+	);
+
+	if (bHasTextureParameter)
+	{
+		UMaterialInstanceDynamic* MID = Cast<UMaterialInstanceDynamic>(MaterialInterface);
+		if (!MID)
+		{
+			MID = CreateAndSetMaterialInstanceDynamic(0);
+		}
+
+		if (ensureMsgf(MID, TEXT("Expected a material instance dynamic")))
+		{
+			MID->SetTextureParameterValue(UE::Composite::Private::CompositeDepthTextureName, DepthTexture);
+		}
+	}
+}
+
+#if WITH_EDITOR
+void UCompositeDepthMeshComponent::PostEditChangeProperty(FPropertyChangedEvent& PropertyChangedEvent)
+{
+	Super::PostEditChangeProperty(PropertyChangedEvent);
+
+	using namespace UE::Composite::Private;
+
+	if (PropertyChangedEvent.GetMemberPropertyName() == GET_MEMBER_NAME_CHECKED(ThisClass, GridResolution))
+	{
+		GenerateGridMesh_Editor();
+	}
+
+	// To be safe, we update the material on any mesh change
+	UpdateMaterial();
+}
+
+void UCompositeDepthMeshComponent::PostTransacted(const FTransactionObjectEvent& InTransactionEvent)
+{
+	Super::PostTransacted(InTransactionEvent);
+
+	if (InTransactionEvent.HasPropertyChanges())
+	{
+		UpdateMaterial();
+	}
+}
+
+void UCompositeDepthMeshComponent::GenerateGridMesh_Editor()
+{
+	// Do we already have this mesh?
+	FString AssetPath = UE::Composite::Private::GetAssetPath(*this);
+	UStaticMesh* NewStaticMesh = Cast<UStaticMesh>(StaticLoadObject(UStaticMesh::StaticClass(), NULL, *AssetPath));
+
+	if (!IsValid(NewStaticMesh))
+	{
+		UE::Geometry::FRectangleMeshGenerator RectGen;
+		RectGen.Width = 1;
+		RectGen.Height = 1;
+		RectGen.WidthVertexCount = GridResolution.X + 1;
+		RectGen.HeightVertexCount = GridResolution.Y + 1;
+		RectGen.bSinglePolyGroup = true;
+		UE::Geometry::FDynamicMesh3 GridMesh(&RectGen.Generate());
+
+		NewStaticMesh = UE::Composite::Private::CreateStaticMeshAsset(&GridMesh, AssetPath);
+	}
+
+	SetStaticMesh(NewStaticMesh);
+}
+#endif
+
+#undef LOCTEXT_NAMESPACE
diff --git a/Engine/Plugins/Compositing/Composite/Source/Composite/Private/Components/CompositeViewProjectionComponent.cpp b/Engine/Plugins/Compositing/Composite/Source/Composite/Private/Components/CompositeViewProjectionComponent.cpp
--- a/Engine/Plugins/Compositing/Composite/Source/Composite/Private/Components/CompositeViewProjectionComponent.cpp
+++ b/Engine/Plugins/Compositing/Composite/Source/Composite/Private/Components/CompositeViewProjectionComponent.cpp
@@ -14,13 +14,133 @@
 
 #define LOCTEXT_NAMESPACE "Composite"
 
+namespace UE
+{
+	namespace Composite
+	{
+		namespace Private
+		{
+			// Replicated code from FViewMatrices::ScreenToClipProjectionMatrix()
+			inline FMatrix ScreenToClipProjectionMatrix(const FMatrix& InProjectionMatrix)
+			{
+				//Screen to clip matrix should not utilise scene depth for the w component in ortho projections, but is needed for perspective
+				const bool bIsPerspectiveProjection = InProjectionMatrix.M[3][3] < 1.0f;
+				
+				if (bIsPerspectiveProjection)
+				{
+					return FMatrix(
+						FPlane(1, 0, 0, 0),
+						FPlane(0, 1, 0, 0),
+						FPlane(0, 0, InProjectionMatrix.M[2][2], 1.0f),
+						FPlane(0, 0, InProjectionMatrix.M[3][2], 0.0f));
+				}
+				else
+				{
+					return FMatrix(
+						FPlane(1, 0, 0, 0),
+						FPlane(0, 1, 0, 0),
+						FPlane(0, 0, InProjectionMatrix.M[2][2], 0.0f),
+						FPlane(0, 0, InProjectionMatrix.M[3][2], 1.0f));
+				}
+			}
+
+			// Replicated code from SceneView.h
+			const FMatrix InvertProjectionMatrix(const FMatrix& M)
+			{
+				if (M.M[1][0] == 0.0f &&
+					M.M[3][0] == 0.0f &&
+					M.M[0][1] == 0.0f &&
+					M.M[3][1] == 0.0f &&
+					M.M[0][2] == 0.0f &&
+					M.M[1][2] == 0.0f &&
+					M.M[0][3] == 0.0f &&
+					M.M[1][3] == 0.0f &&
+					M.M[2][3] == 1.0f &&
+					M.M[3][3] == 0.0f)
+				{
+					// Solve the common case directly with very high precision.
+					/*
+					M =
+					| a | 0 | 0 | 0 |
+					| 0 | b | 0 | 0 |
+					| s | t | c | 1 |
+					| 0 | 0 | d | 0 |
+					*/
+
+					double a = M.M[0][0];
+					double b = M.M[1][1];
+					double c = M.M[2][2];
+					double d = M.M[3][2];
+					double s = M.M[2][0];
+					double t = M.M[2][1];
+
+					return FMatrix(
+						FPlane(1.0 / a, 0.0f, 0.0f, 0.0f),
+						FPlane(0.0f, 1.0 / b, 0.0f, 0.0f),
+						FPlane(0.0f, 0.0f, 0.0f, 1.0 / d),
+						FPlane(-s / a, -t / b, 1.0f, -c / d)
+					);
+				}
+				else
+				{
+					return M.Inverse();
+				}
+			}
+
+			bool UpdateMatrixParameters(
+				UWorld* InWorld,
+				UMaterialParameterCollectionInstance* InMaterialParameterCollectionInstance,
+				const FName& MatrixParameterName,
+				const FMatrix44f& InMatrix)
+			{
+				if (!IsValid(InWorld) || !IsValid(InMaterialParameterCollectionInstance))
+				{
+					return false;
+				}
+
+				const UMaterialParameterCollection* MPC = InMaterialParameterCollectionInstance->GetCollection();
+
+				int32 ParameterProjection = INDEX_NONE;
+				for (int32 ParameterIndex = 0; ParameterIndex < MPC->VectorParameters.Num(); ParameterIndex++)
+				{
+					FName CurrentParameterName = MPC->VectorParameters[ParameterIndex].ParameterName;
+					if (!CurrentParameterName.IsNone())
+					{
+						if (CurrentParameterName == MatrixParameterName)
+						{
+							ParameterProjection = ParameterIndex;
+						}
+					}
+				}
+
+				// Ensure there's space for 4 output vectors for the projection matrix
+				if (ParameterProjection != INDEX_NONE && ParameterProjection + 4 <= MPC->VectorParameters.Num())
+				{
+					// Store the vectors to the collection instance
+					const FLinearColor* MatrixVectors = (const FLinearColor*)&InMatrix;
+					for (int32 ElementIndex = 0; ElementIndex < 4; ElementIndex++)
+					{
+						InMaterialParameterCollectionInstance->SetVectorParameterValue(MPC->VectorParameters[ParameterProjection + ElementIndex].ParameterName, MatrixVectors[ElementIndex]);
+					}
+
+					return true;
+				}
+
+				return false;
+			}
+		}
+	}
+}
+
 UCompositeViewProjectionComponent::UCompositeViewProjectionComponent(const FObjectInitializer& ObjectInitializer)
 	: Super(ObjectInitializer)
 {
 	static ConstructorHelpers::FObjectFinder<UMaterialParameterCollection> DefaultMPC = TEXT("/Composite/Materials/MPC_Composite.MPC_Composite");
 
 	MaterialParameterCollection = DefaultMPC.Object;
-	MatrixParameterName = TEXT("CameraViewProjectionMatrix");
+	ViewProjectionMatrixParameter = TEXT("CameraViewProjectionMatrix");
+	ClipToWorldMatrixParameter = TEXT("CameraClipToWorldMatrix");
+	ScreenToClipMatrixParameter = TEXT("CameraScreenToClipMatrix");
 
 	PrimaryComponentTick.bStartWithTickEnabled = true;
 	PrimaryComponentTick.bCanEverTick = true;
@@ -127,6 +247,8 @@ void UCompositeViewProjectionComponent::OnUnregister()
 
 void UCompositeViewProjectionComponent::UpdateProjection(UWorld* InWorld) const
 {
+	using namespace UE::Composite::Private;
+
 	if (!IsValid(InWorld))
 	{
 		return;
@@ -157,33 +279,20 @@ void UCompositeViewProjectionComponent::UpdateProjection(UWorld* InWorld) const
 
 		if (!LastViewProjectionMatrix.Equals(ViewProjectionMatrix))
 		{
-			int32 ParameterProjection = INDEX_NONE;
-			for (int32 ParameterIndex = 0; ParameterIndex < MPC->VectorParameters.Num(); ParameterIndex++)
-			{
-				FName CurrentParameterName = MPC->VectorParameters[ParameterIndex].ParameterName;
-				if (!CurrentParameterName.IsNone())
-				{
-					if (CurrentParameterName == MatrixParameterName)
-					{
-						ParameterProjection = ParameterIndex;
-					}
-				}
-			}
+			bool bUpdateMPC = false;
+			UMaterialParameterCollectionInstance* InstanceMPC = InWorld->GetParameterCollectionInstance(MPC);
 
-			UMaterialParameterCollectionInstance* MaterialParameterCollectionInstance = GetWorld()->GetParameterCollectionInstance(MaterialParameterCollection);
-			// Ensure there's space for 4 output vectors for the projection matrix
-			if (ParameterProjection != INDEX_NONE && ParameterProjection + 4 <= MaterialParameterCollection->VectorParameters.Num())
-			{
-				FMatrix44f ViewProjectionMatrixf = FMatrix44f(ViewProjectionMatrix);
+			bUpdateMPC |= UpdateMatrixParameters(InWorld, InstanceMPC, ViewProjectionMatrixParameter, FMatrix44f(ViewProjectionMatrix));
 
-				// Store the vectors to the collection instance
-				const FLinearColor* MatrixVectors = (const FLinearColor*)&ViewProjectionMatrixf;
-				for (int32 ElementIndex = 0; ElementIndex < 4; ElementIndex++)
-				{
-					MaterialParameterCollectionInstance->SetVectorParameterValue(MaterialParameterCollection->VectorParameters[ParameterProjection + ElementIndex].ParameterName, MatrixVectors[ElementIndex]);
-				}
+			const FMatrix44f ClipToWorldMatrix{ InvertProjectionMatrix(ProjectionMatrix) * ViewMatrix.Inverse() };
+			bUpdateMPC |= UpdateMatrixParameters(InWorld, InstanceMPC, ClipToWorldMatrixParameter, ClipToWorldMatrix);
+
+			const FMatrix44f ScreenToClipMatrix{ ScreenToClipProjectionMatrix(ProjectionMatrix) };
+			bUpdateMPC |= UpdateMatrixParameters(InWorld, InstanceMPC, ScreenToClipMatrixParameter, ScreenToClipMatrix);
 
-				MaterialParameterCollectionInstance->UpdateRenderState(false);
+			if (bUpdateMPC)
+			{
+				InstanceMPC->UpdateRenderState(false);
 			}
 
 			LastViewProjectionMatrix = ViewProjectionMatrix;
diff --git a/Engine/Plugins/Compositing/Composite/Source/Composite/Private/CompositeDepthMeshActor.cpp b/Engine/Plugins/Compositing/Composite/Source/Composite/Private/CompositeDepthMeshActor.cpp
new file mode 100644
--- /dev/null
+++ b/Engine/Plugins/Compositing/Composite/Source/Composite/Private/CompositeDepthMeshActor.cpp
@@ -0,0 +1,15 @@
+// Copyright Epic Games, Inc. All Rights Reserved.
+
+#include "CompositeDepthMeshActor.h"
+
+#include "Components/CompositeDepthMeshComponent.h"
+
+ACompositeDepthMeshActor::ACompositeDepthMeshActor(const FObjectInitializer& ObjectInitializer)
+	: Super(ObjectInitializer)
+{
+	CompositeDepthMeshComponent = CreateDefaultSubobject<UCompositeDepthMeshComponent>(TEXT("DefaultCompositeMeshComponent"));
+	SetRootComponent(CompositeDepthMeshComponent);
+}
+
+ACompositeDepthMeshActor::~ACompositeDepthMeshActor() = default;
+
diff --git a/Engine/Plugins/Compositing/Composite/Source/Composite/Public/Components/CompositeDepthMeshComponent.h b/Engine/Plugins/Compositing/Composite/Source/Composite/Public/Components/CompositeDepthMeshComponent.h
new file mode 100644
--- /dev/null
+++ b/Engine/Plugins/Compositing/Composite/Source/Composite/Public/Components/CompositeDepthMeshComponent.h
@@ -0,0 +1,46 @@
+// Copyright Epic Games, Inc. All Rights Reserved.
+
+#pragma once
+
+#include "CoreTypes.h"
+#include "Components/StaticMeshComponent.h"
+
+#include "CompositeDepthMeshComponent.generated.h"
+
+#define UE_API COMPOSITE_API
+
+UCLASS(MinimalAPI, ClassGroup = Composite, editinlinenew, HideCategories = (Transform, Tags, Cooking, Physics, LOD, AssetUserData, Navigation), PrioritizeCategories = ("Composite"), meta = (BlueprintSpawnableComponent))
+class UCompositeDepthMeshComponent : public UStaticMeshComponent
+{
+	GENERATED_UCLASS_BODY()
+
+public:
+
+	//~ Begin UObject interface
+	UE_API virtual void PostLoad() override;
+	//~ End UObject interface
+
+#if WITH_EDITOR
+	UE_API virtual void PostEditChangeProperty(FPropertyChangedEvent& PropertyChangedEvent) override;
+	UE_API virtual void PostTransacted(const FTransactionObjectEvent& InTransactionEvent) override;
+#endif
+
+	/* Planar grid mesh resolution. */
+	UPROPERTY(EditAnywhere, BlueprintReadOnly, Category = "Composite")
+	FIntVector2 GridResolution;
+
+	/* Depth texture sent to material. */
+	UPROPERTY(EditAnywhere, Category = "Composite")
+	TObjectPtr<UTexture> DepthTexture;
+
+private:
+	/* Convert material to MID and update its depth texture parameter. */
+	void UpdateMaterial();
+
+#if WITH_EDITOR
+	/* Generate our static grid mesh based on the resolution parameters. */
+	void GenerateGridMesh_Editor();
+#endif
+};
+
+#undef UE_API
diff --git a/Engine/Plugins/Compositing/Composite/Source/Composite/Public/Components/CompositeViewProjectionComponent.h b/Engine/Plugins/Compositing/Composite/Source/Composite/Public/Components/CompositeViewProjectionComponent.h
--- a/Engine/Plugins/Compositing/Composite/Source/Composite/Public/Components/CompositeViewProjectionComponent.h
+++ b/Engine/Plugins/Compositing/Composite/Source/Composite/Public/Components/CompositeViewProjectionComponent.h
@@ -46,9 +46,17 @@ public:
 	UPROPERTY(EditAnywhere, BlueprintReadWrite, Interp, Category = "CameraProjection")
 	TObjectPtr<UMaterialParameterCollection> MaterialParameterCollection;
 
-	/** Parameter name of the first element of the transform in the Material Parameter Collection set above.  Requires space for 4 vectors. */
+	/** Parameter name of the first element of the view projection matrix in the Material Parameter Collection set above.  Requires space for 4 vectors. */
 	UPROPERTY(EditAnywhere, BlueprintReadWrite, Interp, Category = "CameraProjection")
-	FName MatrixParameterName;
+	FName ViewProjectionMatrixParameter;
+
+	/** Parameter name of the first element of the inverse view projection matrix in the Material Parameter Collection set above.  Requires space for 4 vectors. */
+	UPROPERTY(EditAnywhere, BlueprintReadWrite, Interp, Category = "CameraProjection", AdvancedDisplay)
+	FName ClipToWorldMatrixParameter;
+
+	/** Parameter name of the first element of the screen to clip matrix in the Material Parameter Collection set above.  Requires space for 4 vectors. */
+	UPROPERTY(EditAnywhere, BlueprintReadWrite, Interp, Category = "CameraProjection", AdvancedDisplay)
+	FName ScreenToClipMatrixParameter;
 
 public:
 	/** Force update the Material Parameter Collection this frame. */
diff --git a/Engine/Plugins/Compositing/Composite/Source/Composite/Public/CompositeDepthMeshActor.h b/Engine/Plugins/Compositing/Composite/Source/Composite/Public/CompositeDepthMeshActor.h
new file mode 100644
--- /dev/null
+++ b/Engine/Plugins/Compositing/Composite/Source/Composite/Public/CompositeDepthMeshActor.h
@@ -0,0 +1,30 @@
+// Copyright Epic Games, Inc. All Rights Reserved.
+
+#pragma once
+
+#include "CoreTypes.h"
+#include "GameFramework/Actor.h"
+
+#include "CompositeDepthMeshActor.generated.h"
+
+#define UE_API COMPOSITE_API
+
+class UCompositeDepthMeshComponent;
+
+/** Convenience compositing depth mesh actor. */
+UCLASS(MinimalAPI)
+class ACompositeDepthMeshActor : public AActor
+{
+	GENERATED_BODY()
+
+public:
+	UE_API ACompositeDepthMeshActor(const FObjectInitializer& ObjectInitializer);
+	UE_API ~ACompositeDepthMeshActor();
+
+private:
+	/** Composite mesh component. */
+	UPROPERTY(VisibleAnywhere, BlueprintReadOnly, Category = "Composite", meta = (AllowPrivateAccess = "true"))
+	TObjectPtr<UCompositeDepthMeshComponent> CompositeDepthMeshComponent;
+};
+
+#undef UE_API
diff --git a/Engine/Plugins/Compositing/Composite/Source/Composite/Public/CompositeMeshActor.h b/Engine/Plugins/Compositing/Composite/Source/Composite/Public/CompositeMeshActor.h
--- a/Engine/Plugins/Compositing/Composite/Source/Composite/Public/CompositeMeshActor.h
+++ b/Engine/Plugins/Compositing/Composite/Source/Composite/Public/CompositeMeshActor.h
@@ -11,7 +11,7 @@
 
 class UCompositeMeshComponent;
 
-/** Convenience compositing mesh actor, which the composite actor can reference. */
+/** Convenience compositing mesh actor. */
 UCLASS(MinimalAPI)
 class ACompositeMeshActor : public AActor
 {
@@ -25,10 +25,6 @@ private:
 	/** Composite mesh component. */
 	UPROPERTY(VisibleAnywhere, BlueprintReadOnly, Category = "Composite", meta = (AllowPrivateAccess = "true"))
 	TObjectPtr<UCompositeMeshComponent> CompositeMeshComponent;
-
-public:
-	/** Returns CompositeMeshComponent subobject **/
-	UE_API class UCompositeMeshComponent* GetCaptureComponent2D() const { return CompositeMeshComponent; }
 };
 
 #undef UE_API
diff --git a/Engine/Plugins/Compositing/Composite/Source/CompositeEditor/Private/CompositeEditorModule.cpp b/Engine/Plugins/Compositing/Composite/Source/CompositeEditor/Private/CompositeEditorModule.cpp
--- a/Engine/Plugins/Compositing/Composite/Source/CompositeEditor/Private/CompositeEditorModule.cpp
+++ b/Engine/Plugins/Compositing/Composite/Source/CompositeEditor/Private/CompositeEditorModule.cpp
@@ -12,6 +12,7 @@
 #include "LevelEditor.h"
 #include "Modules/ModuleManager.h"
 #include "Widgets/Notifications/SNotificationList.h"
+#include "Components/CompositeDepthMeshComponent.h"
 #include "CompositeActor.h"
 #include "CompositeEditorCommands.h"
 #include "CompositeEditorStyle.h"
@@ -21,6 +22,7 @@
 #include "Customizations/CompositeLayerSceneCaptureCustomization.h"
 #include "Customizations/CompositeLayerShadowReflectionCustomization.h"
 #include "Customizations/CompositeLayerSingleLightShadowCustomization.h"
+#include "Customizations/CompositeDepthMeshComponentCustomization.h"
 #include "Layers/CompositeLayerPlate.h"
 #include "Layers/CompositeLayerSceneCapture.h"
 #include "Layers/CompositeLayerShadowReflection.h"
@@ -251,6 +253,11 @@ void FCompositeEditorModule::RegisterCustomizations()
 		ACompositeActor::StaticClass()->GetFName(),
 		FOnGetDetailCustomizationInstance::CreateStatic(&FCompositeActorCustomization::MakeInstance)
 	);
+
+	PropertyModule.RegisterCustomClassLayout(
+		UCompositeDepthMeshComponent::StaticClass()->GetFName(),
+		FOnGetDetailCustomizationInstance::CreateStatic(&FCompositeDepthMeshComponentCustomization::MakeInstance)
+	);
 }
 
 void FCompositeEditorModule::UnregisterCustomizations()
diff --git a/Engine/Plugins/Compositing/Composite/Source/CompositeEditor/Private/Customizations/CompositeDepthMeshComponentCustomization.cpp b/Engine/Plugins/Compositing/Composite/Source/CompositeEditor/Private/Customizations/CompositeDepthMeshComponentCustomization.cpp
new file mode 100644
--- /dev/null
+++ b/Engine/Plugins/Compositing/Composite/Source/CompositeEditor/Private/Customizations/CompositeDepthMeshComponentCustomization.cpp
@@ -0,0 +1,20 @@
+// Copyright Epic Games, Inc. All Rights Reserved.
+
+#include "CompositeDepthMeshComponentCustomization.h"
+
+#include "DetailLayoutBuilder.h"
+
+#define LOCTEXT_NAMESPACE "FCompositeDepthMeshComponentCustomization"
+
+TSharedRef<IDetailCustomization> FCompositeDepthMeshComponentCustomization::MakeInstance()
+{
+	return MakeShared<FCompositeDepthMeshComponentCustomization>();
+}
+
+void FCompositeDepthMeshComponentCustomization::CustomizeDetails(IDetailLayoutBuilder& DetailBuilder)
+{
+	// Promote "Composite" category
+	IDetailCategoryBuilder& CompositeCategory = DetailBuilder.EditCategory("Composite", FText::GetEmpty(), ECategoryPriority::Important);
+}
+
+#undef LOCTEXT_NAMESPACE
diff --git a/Engine/Plugins/Compositing/Composite/Source/CompositeEditor/Private/Customizations/CompositeDepthMeshComponentCustomization.h b/Engine/Plugins/Compositing/Composite/Source/CompositeEditor/Private/Customizations/CompositeDepthMeshComponentCustomization.h
new file mode 100644
--- /dev/null
+++ b/Engine/Plugins/Compositing/Composite/Source/CompositeEditor/Private/Customizations/CompositeDepthMeshComponentCustomization.h
@@ -0,0 +1,19 @@
+// Copyright Epic Games, Inc. All Rights Reserved.
+
+#pragma once
+
+#include "IDetailCustomization.h"
+
+/**
+ * Customization for the composite depth mesh component.
+ */
+class FCompositeDepthMeshComponentCustomization : public IDetailCustomization
+{
+public:
+	/** Makes a new instance of the details customization */
+	static TSharedRef<IDetailCustomization> MakeInstance();
+
+	// IDetailCustomization interface
+	virtual void CustomizeDetails(IDetailLayoutBuilder& DetailBuilder) override;
+	// End of IDetailCustomization interface
+};
