Composure: Further fix to CL 47600697.

* If a layer was removed its callback registration was undone. Undo'ing this removal would fail to restore this end-of-frame callback.
* Instead, we only remove the callback on destruction, and prevent callbacks on removed or destroyed objects with extra conditions.

#Synced-CL 48290222

diff --git a/Engine/Plugins/Compositing/Composite/Source/Composite/Private/Layers/CompositeLayerBase.cpp b/Engine/Plugins/Compositing/Composite/Source/Composite/Private/Layers/CompositeLayerBase.cpp
--- a/Engine/Plugins/Compositing/Composite/Source/Composite/Private/Layers/CompositeLayerBase.cpp
+++ b/Engine/Plugins/Compositing/Composite/Source/Composite/Private/Layers/CompositeLayerBase.cpp
@@ -25,13 +25,16 @@ UCompositeLayerBase::UCompositeLayerBase(const FObjectInitializer& ObjectInitial
 #endif
 }
 
-UCompositeLayerBase::~UCompositeLayerBase()
+UCompositeLayerBase::~UCompositeLayerBase() = default;
+
+void UCompositeLayerBase::BeginDestroy()
 {
+	Super::BeginDestroy();
+
 	constexpr bool bEnabled = false;
 	RegisterEndOfFrameUpdate(bEnabled);
 }
 
-
 #if WITH_EDITOR
 bool UCompositeLayerBase::CanEditChange(const FProperty* InProperty) const
 {
@@ -71,12 +74,6 @@ void UCompositeLayerBase::PostEditChangeProperty(FPropertyChangedEvent& Property
 }
 #endif //WITH_EDITOR
 
-void UCompositeLayerBase::OnRemoved(ACompositeActor* LastOwner)
-{
-	constexpr bool bEnabled = false;
-	RegisterEndOfFrameUpdate(bEnabled);
-}
-
 bool UCompositeLayerBase::IsEnabled() const
 {
 	return bIsEnabled;
@@ -93,7 +90,7 @@ void UCompositeLayerBase::RegisterEndOfFrameUpdate(bool bInEnabled)
 	{
 		if (!OnWorldPreSendAllEndOfFrameUpdatesHandle.IsValid())
 		{
-			OnWorldPreSendAllEndOfFrameUpdatesHandle = FWorldDelegates::OnWorldPreSendAllEndOfFrameUpdates.AddUObject(this, &UCompositeLayerBase::OnEndOfFrameUpdate);
+			OnWorldPreSendAllEndOfFrameUpdatesHandle = FWorldDelegates::OnWorldPreSendAllEndOfFrameUpdates.AddUObject(this, &UCompositeLayerBase::ConditionalOnEndOfFrameUpdate);
 		}
 	}
 	else
@@ -169,6 +166,23 @@ UE::CompositeCore::FPassInputDecl UCompositeLayerBase::GetDefaultSecondInput(con
 	return InputDecl;
 }
 
+void UCompositeLayerBase::ConditionalOnEndOfFrameUpdate(UWorld* InWorld)
+{
+	/**
+	 * Only call the virtual end-of-frame method if our layer hasn't been disabled or removed.
+	 * This is simpler than unregistering and re-registering if the removal is undone.
+	*/
+	if (IsEnabled())
+	{
+		ACompositeActor* CompositeActor = GetTypedOuter<ACompositeActor>();
+
+		if (IsValid(CompositeActor) && CompositeActor->GetCompositeLayers().Contains(this))
+		{
+			OnEndOfFrameUpdate(InWorld);
+		}
+	}
+}
+
 UE::CompositeCore::ResourceId UCompositeLayerBase::FTraversalContext::FindOrCreateExternalTexture(TWeakObjectPtr<UTexture> InTexture, UE::CompositeCore::FResourceMetadata InMetadata)
 {
 	int32 TextureIndex = ExternalTextures.IndexOfByPredicate([&InTexture, &InMetadata](const UE::CompositeCore::FExternalTexture& InExternalTexture)
diff --git a/Engine/Plugins/Compositing/Composite/Source/Composite/Private/Layers/CompositeLayerPlate.cpp b/Engine/Plugins/Compositing/Composite/Source/Composite/Private/Layers/CompositeLayerPlate.cpp
--- a/Engine/Plugins/Compositing/Composite/Source/Composite/Private/Layers/CompositeLayerPlate.cpp
+++ b/Engine/Plugins/Compositing/Composite/Source/Composite/Private/Layers/CompositeLayerPlate.cpp
@@ -190,8 +190,6 @@ void UCompositeLayerPlate::PostLoad()
 
 void UCompositeLayerPlate::OnRemoved(ACompositeActor* LastOwner)
 {
-	Super::OnRemoved(LastOwner);
-
 	FCompositeRenderTargetPool::Get().ReleaseTarget(MediaRenderTarget);
 	FCompositeRenderTargetPool::Get().ReleaseTarget(SceneRenderTarget);
 
diff --git a/Engine/Plugins/Compositing/Composite/Source/Composite/Private/Layers/CompositeLayerSceneCapture.cpp b/Engine/Plugins/Compositing/Composite/Source/Composite/Private/Layers/CompositeLayerSceneCapture.cpp
--- a/Engine/Plugins/Compositing/Composite/Source/Composite/Private/Layers/CompositeLayerSceneCapture.cpp
+++ b/Engine/Plugins/Compositing/Composite/Source/Composite/Private/Layers/CompositeLayerSceneCapture.cpp
@@ -27,8 +27,6 @@ UCompositeLayerSceneCapture::~UCompositeLayerSceneCapture() = default;
 
 void UCompositeLayerSceneCapture::OnRemoved(ACompositeActor* LastOwner)
 {
-	Super::OnRemoved(LastOwner);
-
 	// Note: We don't use IsValid(..) since destruction should proceed even if the actor is pending kill.
 	if (LastOwner != nullptr)
 	{
diff --git a/Engine/Plugins/Compositing/Composite/Source/Composite/Private/Layers/CompositeLayerShadowReflection.cpp b/Engine/Plugins/Compositing/Composite/Source/Composite/Private/Layers/CompositeLayerShadowReflection.cpp
--- a/Engine/Plugins/Compositing/Composite/Source/Composite/Private/Layers/CompositeLayerShadowReflection.cpp
+++ b/Engine/Plugins/Compositing/Composite/Source/Composite/Private/Layers/CompositeLayerShadowReflection.cpp
@@ -25,8 +25,6 @@ UCompositeLayerShadowReflection::~UCompositeLayerShadowReflection() = default;
 
 void UCompositeLayerShadowReflection::OnRemoved(ACompositeActor* LastOwner)
 {
-	Super::OnRemoved(LastOwner);
-
 	// Note: We don't use IsValid(..) since destruction should proceed even if the actor is pending kill.
 	if (LastOwner != nullptr)
 	{
diff --git a/Engine/Plugins/Compositing/Composite/Source/Composite/Private/Layers/CompositeLayerSingleLightShadow.cpp b/Engine/Plugins/Compositing/Composite/Source/Composite/Private/Layers/CompositeLayerSingleLightShadow.cpp
--- a/Engine/Plugins/Compositing/Composite/Source/Composite/Private/Layers/CompositeLayerSingleLightShadow.cpp
+++ b/Engine/Plugins/Compositing/Composite/Source/Composite/Private/Layers/CompositeLayerSingleLightShadow.cpp
@@ -172,8 +172,6 @@ UCompositeLayerSingleLightShadow::UCompositeLayerSingleLightShadow(const FObject
 
 void UCompositeLayerSingleLightShadow::OnRemoved(ACompositeActor* LastOwner)
 {
-	Super::OnRemoved(LastOwner);
-
 	// Note: We don't use IsValid(..) since destruction should proceed even if the actor is pending kill.
 	if (LastOwner != nullptr)
 	{
diff --git a/Engine/Plugins/Compositing/Composite/Source/Composite/Public/Layers/CompositeLayerBase.h b/Engine/Plugins/Compositing/Composite/Source/Composite/Public/Layers/CompositeLayerBase.h
--- a/Engine/Plugins/Compositing/Composite/Source/Composite/Public/Layers/CompositeLayerBase.h
+++ b/Engine/Plugins/Compositing/Composite/Source/Composite/Public/Layers/CompositeLayerBase.h
@@ -59,10 +59,14 @@ public:
 	/** Destructor. */
 	UE_API virtual ~UCompositeLayerBase();
 
+	//~ Begin UObject Interface
+	UE_API virtual void BeginDestroy() override;
+
 #if WITH_EDITOR
 	UE_API virtual bool CanEditChange(const FProperty* InProperty) const;
 	UE_API virtual void PostEditChangeProperty(FPropertyChangedEvent& PropertyChangedEvent) override;
 #endif
+	//~ End UObject Interface
 
 	/**
 	 * End-of-frame update callback to inherit.
@@ -82,7 +86,7 @@ public:
 	*
 	* @param LastOwner - Last composite actor owner, from which the layer was removed. Ignore if null.
 	*/
-	UE_API virtual void OnRemoved(ACompositeActor* LastOwner);
+	UE_API virtual void OnRemoved(ACompositeActor* LastOwner) {}
 
 	/** Getter function to override, returning pass layer proxies to be passed to the render thread. (Proxy objects should be allocated from the provided allocator.) */
 	UE_API virtual bool GetProxy(FTraversalContext& InContext, FSceneRenderingBulkObjectAllocator& InFrameAllocator, FCompositeCorePassProxy*& OutProxy) const { return false; };
@@ -135,6 +139,9 @@ protected:
 	UE::CompositeCore::FPassInputDecl GetDefaultSecondInput(const FTraversalContext& InContext) const;
 
 private:
+	/** Private end of frame callback implementation to condition the inherited function. */
+	void ConditionalOnEndOfFrameUpdate(UWorld* InWorld);
+
 	/** Post-viewport update callback. */
 	FDelegateHandle OnWorldPreSendAllEndOfFrameUpdatesHandle;
 };
