Composite: Fixes to single light shadow layer.

* [Breaking Change] Use of shadow caster actors is now mandatory, but yields cleaner results automatically. Users are now in control of which meshes should cast shadows, and the UX matches the shadow/reflection layer. More specifically, the shadow map is now a show-only custom render pass.
* Shadow bias parameter is now deprecated since shadow casters can never be catchers.

#Synced-CL 48399844

diff --git a/Engine/Plugins/Compositing/Composite/Source/Composite/Private/Layers/CompositeLayerBase.cpp b/Engine/Plugins/Compositing/Composite/Source/Composite/Private/Layers/CompositeLayerBase.cpp
--- a/Engine/Plugins/Compositing/Composite/Source/Composite/Private/Layers/CompositeLayerBase.cpp
+++ b/Engine/Plugins/Compositing/Composite/Source/Composite/Private/Layers/CompositeLayerBase.cpp
@@ -4,6 +4,8 @@
 
 #include "CompositeActor.h"
 #include "CompositeRenderTargetPool.h"
+#include "Components/PrimitiveComponent.h"
+#include "Components/SceneCaptureComponent2D.h"
 
 UCompositeLayerBase::UCompositeLayerBase(const FObjectInitializer& ObjectInitializer)
 	: Super(ObjectInitializer)
@@ -84,6 +86,37 @@ void UCompositeLayerBase::SetEnabled(bool bInEnabled)
 	bIsEnabled = bInEnabled;
 }
 
+void UCompositeLayerBase::UpdateSceneCaptureComponents(USceneCaptureComponent2D& SceneCaptureComponent, TArrayView<TSoftObjectPtr<AActor>> InActors)
+{
+	auto UpdateComponentsFn = [InActors](TArray<TWeakObjectPtr<UPrimitiveComponent>>& InComponents) -> void
+		{
+			InComponents.Reset(InActors.Num());
+
+			for (const TSoftObjectPtr<AActor>& SoftActor : InActors)
+			{
+				if (const AActor* Actor = SoftActor.Get())
+				{
+					for (UActorComponent* Component : Actor->GetComponents())
+					{
+						if (UPrimitiveComponent* PrimitiveComponent = Cast<UPrimitiveComponent>(Component))
+						{
+							InComponents.Add(PrimitiveComponent);
+						}
+					}
+				}
+			}
+		};
+	
+	if (SceneCaptureComponent.PrimitiveRenderMode == ESceneCapturePrimitiveRenderMode::PRM_UseShowOnlyList)
+	{
+		UpdateComponentsFn(SceneCaptureComponent.ShowOnlyComponents);
+	}
+	else
+	{
+		UpdateComponentsFn(SceneCaptureComponent.HiddenComponents);
+	}
+}
+
 void UCompositeLayerBase::RegisterEndOfFrameUpdate(bool bInEnabled)
 {
 	if (bInEnabled)
diff --git a/Engine/Plugins/Compositing/Composite/Source/Composite/Private/Layers/CompositeLayerSceneCapture.cpp b/Engine/Plugins/Compositing/Composite/Source/Composite/Private/Layers/CompositeLayerSceneCapture.cpp
--- a/Engine/Plugins/Compositing/Composite/Source/Composite/Private/Layers/CompositeLayerSceneCapture.cpp
+++ b/Engine/Plugins/Compositing/Composite/Source/Composite/Private/Layers/CompositeLayerSceneCapture.cpp
@@ -108,7 +108,7 @@ void UCompositeLayerSceneCapture::OnEndOfFrameUpdate(UWorld* InWorld)
 	USceneCaptureComponent2D* SceneCaptureComponent = GetSceneCaptureComponent();
 	if (IsValid(SceneCaptureComponent))
 	{
-		UpdateSceneCaptureComponents(*SceneCaptureComponent);
+		UpdateSceneCaptureComponents(*SceneCaptureComponent, Actors);
 	}
 }
 
@@ -169,7 +169,7 @@ void UCompositeLayerSceneCapture::SetActors(TArray<TSoftObjectPtr<AActor>> InAct
 		// First, we empty show-only actors since our previous logic used them
 		SceneCaptureComponent->ShowOnlyActors.Empty();
 
-		UpdateSceneCaptureComponents(*SceneCaptureComponent);
+		UpdateSceneCaptureComponents(*SceneCaptureComponent, Actors);
 	}
 }
 
@@ -250,22 +250,3 @@ void UCompositeLayerSceneCapture::RestorePrimitiveVisibilityState() const
 	UpdatePrimitiveVisibilityState(CDO->bVisibleInSceneCaptureOnly);
 }
 
-void UCompositeLayerSceneCapture::UpdateSceneCaptureComponents(USceneCaptureComponent2D& SceneCaptureComponent) const
-{
-	SceneCaptureComponent.ShowOnlyComponents.Reset(Actors.Num());
-
-	for (const TSoftObjectPtr<AActor>& SoftActor : Actors)
-	{
-		if (const AActor* Actor = SoftActor.Get())
-		{
-			for (UActorComponent* Component : Actor->GetComponents())
-			{
-				if (UPrimitiveComponent* PrimitiveComponent = Cast<UPrimitiveComponent>(Component))
-				{
-					SceneCaptureComponent.ShowOnlyComponents.Add(PrimitiveComponent);
-				}
-			}
-		}
-	}
-}
-
diff --git a/Engine/Plugins/Compositing/Composite/Source/Composite/Private/Layers/CompositeLayerShadowReflection.cpp b/Engine/Plugins/Compositing/Composite/Source/Composite/Private/Layers/CompositeLayerShadowReflection.cpp
--- a/Engine/Plugins/Compositing/Composite/Source/Composite/Private/Layers/CompositeLayerShadowReflection.cpp
+++ b/Engine/Plugins/Compositing/Composite/Source/Composite/Private/Layers/CompositeLayerShadowReflection.cpp
@@ -115,7 +115,7 @@ void UCompositeLayerShadowReflection::OnEndOfFrameUpdate(UWorld* InWorld)
 	*/
 	if (CachedSceneCaptures[1].IsValid())
 	{
-		UpdateSceneCaptureComponents(*CachedSceneCaptures[1]);
+		UpdateSceneCaptureComponents(*CachedSceneCaptures[1], Actors);
 	}
 }
 
@@ -211,7 +211,7 @@ void UCompositeLayerShadowReflection::SetActors(TArray<TSoftObjectPtr<AActor>> I
 	// First, we empty hidden actors since our previous logic used them
 	CachedSceneCaptures[1]->HiddenActors.Empty();
 
-	UpdateSceneCaptureComponents(*CachedSceneCaptures[1]);
+	UpdateSceneCaptureComponents(*CachedSceneCaptures[1], Actors);
 }
 
 void UCompositeLayerShadowReflection::UpdatePrimitiveVisibilityState(
@@ -265,25 +265,6 @@ void UCompositeLayerShadowReflection::RestorePrimitiveVisibilityState() const
 	UpdatePrimitiveVisibilityState(CDO->bHiddenInSceneCapture, CDO->bAffectIndirectLightingWhileHidden, CDO->bCastHiddenShadow);
 }
 
-void UCompositeLayerShadowReflection::UpdateSceneCaptureComponents(USceneCaptureComponent2D& SceneCaptureComponent) const
-{
-	SceneCaptureComponent.HiddenComponents.Reset(Actors.Num());
-
-	for (const TSoftObjectPtr<AActor>& SoftActor : Actors)
-	{
-		if (const AActor* Actor = SoftActor.Get())
-		{
-			for (UActorComponent* Component : Actor->GetComponents())
-			{
-				if (UPrimitiveComponent* PrimitiveComponent = Cast<UPrimitiveComponent>(Component))
-				{
-					SceneCaptureComponent.HiddenComponents.Add(PrimitiveComponent);
-				}
-			}
-		}
-	}
-}
-
 TStaticArray<TWeakObjectPtr<UCompositeShadowReflectionCatcherComponent>, 2> UCompositeLayerShadowReflection::FindOrCreateSceneCapturePair(ACompositeActor& InOuter) const
 {
 	return {
diff --git a/Engine/Plugins/Compositing/Composite/Source/Composite/Private/Layers/CompositeLayerSingleLightShadow.cpp b/Engine/Plugins/Compositing/Composite/Source/Composite/Private/Layers/CompositeLayerSingleLightShadow.cpp
--- a/Engine/Plugins/Compositing/Composite/Source/Composite/Private/Layers/CompositeLayerSingleLightShadow.cpp
+++ b/Engine/Plugins/Compositing/Composite/Source/Composite/Private/Layers/CompositeLayerSingleLightShadow.cpp
@@ -214,9 +214,14 @@ void UCompositeLayerSingleLightShadow::OnEndOfFrameUpdate(UWorld* InWorld)
 	* primitive components may often be recreated (see RerunConstructionScripts(), FStaticMeshComponentRecreateRenderStateContext, dynamic
 	* text primitives, etc). To prevent issues with destroyed primitives disappearing in our scene captures, we re-register them every frame.
 	*/
+	if (CachedShadowMapCapture.IsValid())
+	{
+		UpdateSceneCaptureComponents(*CachedShadowMapCapture, ShadowCastingActors);
+	}
+
 	if (CachedSceneDepthCapture.IsValid())
 	{
-		UpdateSceneCaptureComponents(*CachedSceneDepthCapture);
+		UpdateSceneCaptureComponents(*CachedSceneDepthCapture, ShadowCastingActors);
 	}
 }
 
@@ -278,12 +283,8 @@ bool UCompositeLayerSingleLightShadow::GetProxy(FTraversalContext& InContext, FS
 		PassInputs[0].Set<FPassExternalResourceDesc>({ SceneDepthId });
 		PassInputs[1].Set<FPassExternalResourceDesc>({ ShadowMapId });
 
-		static const TConsoleVariableData<float>* CVarDepthBias = IConsoleManager::Get().FindTConsoleVariableDataFloat(TEXT("r.Shadow.CSMDepthBias"));
-		float ShadowDepthBias = CVarDepthBias ? CVarDepthBias->GetValueOnGameThread() : 10.0f;
-		ShadowDepthBias *= ShadowBias;
-
 		ShadowProxy = InFrameAllocator.Create<FCompositeSingleLightShadowProxy>(PassInputs);
-		ShadowProxy->ShadowStrengthBias = FVector4f(ShadowStrength, ShadowDepthBias, 0.0f, 0.0f);
+		ShadowProxy->ShadowStrengthBias = FVector4f(ShadowStrength, 0.0f, 0.0f, 0.0f);
 		GetShadowMatrices(*CachedShadowMapCapture, ShadowProxy->ShadowMatrix, ShadowProxy->ShadowInvDeviceZToWorldZ);
 	}
 	ShadowProxy->LightWorldDirection = FVector3f(Light->GetTransform().TransformVector(FVector(1, 0, 0)));
@@ -312,17 +313,27 @@ void UCompositeLayerSingleLightShadow::SetShadowCastingActors(TArray<TSoftObject
 		return;
 	}
 
+	if (!CachedShadowMapCapture.IsValid())
+	{
+		CachedShadowMapCapture = FindOrCreateShadowMapCapture(*CompositeActor);
+	}
+
 	if (!CachedSceneDepthCapture.IsValid())
 	{
 		CachedSceneDepthCapture = FindOrCreateSceneDepthCapture(*CompositeActor);
 	}
+
+	if (CachedShadowMapCapture.IsValid())
+	{
+		UpdateSceneCaptureComponents(*CachedShadowMapCapture, ShadowCastingActors);
+	}
 	
 	if (CachedSceneDepthCapture.IsValid())
 	{
 		// First, we empty hidden actors since our previous logic used them
 		CachedSceneDepthCapture->HiddenActors.Empty();
 
-		UpdateSceneCaptureComponents(*CachedSceneDepthCapture);
+		UpdateSceneCaptureComponents(*CachedSceneDepthCapture, ShadowCastingActors);
 	}
 }
 
@@ -347,7 +358,7 @@ UCompositeSceneCapture2DComponent* UCompositeLayerSingleLightShadow::FindOrCreat
 
 	if (ensure(IsValid(ShadowMapCapture)))
 	{
-		ShadowMapCapture->PrimitiveRenderMode = ESceneCapturePrimitiveRenderMode::PRM_RenderScenePrimitives;
+		ShadowMapCapture->PrimitiveRenderMode = ESceneCapturePrimitiveRenderMode::PRM_UseShowOnlyList;
 		ShadowMapCapture->CaptureSource = SCS_SceneDepth;
 		ShadowMapCapture->bRenderInMainRenderer = true;
 		ShadowMapCapture->bMainViewFamily = false;
@@ -361,26 +372,6 @@ UCompositeSceneCapture2DComponent* UCompositeLayerSingleLightShadow::FindOrCreat
 	return ShadowMapCapture;
 }
 
-void UCompositeLayerSingleLightShadow::UpdateSceneCaptureComponents(USceneCaptureComponent2D& SceneCaptureComponent) const
-{
-	// Collect & set hidden primitive components
-	SceneCaptureComponent.HiddenComponents.Reset(ShadowCastingActors.Num());
-
-	for (const TSoftObjectPtr<AActor>& SoftActor : ShadowCastingActors)
-	{
-		if (const AActor* Actor = SoftActor.Get())
-		{
-			for (UActorComponent* Component : Actor->GetComponents())
-			{
-				if (UPrimitiveComponent* PrimitiveComponent = Cast<UPrimitiveComponent>(Component))
-				{
-					SceneCaptureComponent.HiddenComponents.Add(PrimitiveComponent);
-				}
-			}
-		}
-	}
-}
-
 void UCompositeLayerSingleLightShadow::GetShadowMatrices(UCompositeSceneCapture2DComponent& ShadowMapCapture, FMatrix44f& OutShadowMatrix, FVector4f& OutShadowInvDeviceZToWorldZ) const
 {
 	FMinimalViewInfo ViewInfo;
diff --git a/Engine/Plugins/Compositing/Composite/Source/Composite/Public/Layers/CompositeLayerBase.h b/Engine/Plugins/Compositing/Composite/Source/Composite/Public/Layers/CompositeLayerBase.h
--- a/Engine/Plugins/Compositing/Composite/Source/Composite/Public/Layers/CompositeLayerBase.h
+++ b/Engine/Plugins/Compositing/Composite/Source/Composite/Public/Layers/CompositeLayerBase.h
@@ -100,6 +100,9 @@ public:
 	UE_API virtual void SetEnabled(bool bInEnabled);
 
 protected:
+	/** Convenience function to update scene capture primitive components, depending on the scene capture primitive render mode. */
+	UE_API static void UpdateSceneCaptureComponents(USceneCaptureComponent2D& SceneCaptureComponent, TArrayView<TSoftObjectPtr<AActor>> InActors);
+
 	/** Activate the end-of-frame update callback. */
 	UE_API void RegisterEndOfFrameUpdate(bool bInEnabled);
 
diff --git a/Engine/Plugins/Compositing/Composite/Source/Composite/Public/Layers/CompositeLayerSceneCapture.h b/Engine/Plugins/Compositing/Composite/Source/Composite/Public/Layers/CompositeLayerSceneCapture.h
--- a/Engine/Plugins/Compositing/Composite/Source/Composite/Public/Layers/CompositeLayerSceneCapture.h
+++ b/Engine/Plugins/Compositing/Composite/Source/Composite/Public/Layers/CompositeLayerSceneCapture.h
@@ -85,9 +85,6 @@ private:
 	/** Restore scene capture visibility on registered mesh actors primitives. */
 	void RestorePrimitiveVisibilityState()  const;
 
-	/** Update scene capture show only components. */
-	void UpdateSceneCaptureComponents(USceneCaptureComponent2D& SceneCaptureComponent) const;
-
 	/**
 	 * List of actors to be rendered from the scene capture.
 	 * 
diff --git a/Engine/Plugins/Compositing/Composite/Source/Composite/Public/Layers/CompositeLayerShadowReflection.h b/Engine/Plugins/Compositing/Composite/Source/Composite/Public/Layers/CompositeLayerShadowReflection.h
--- a/Engine/Plugins/Compositing/Composite/Source/Composite/Public/Layers/CompositeLayerShadowReflection.h
+++ b/Engine/Plugins/Compositing/Composite/Source/Composite/Public/Layers/CompositeLayerShadowReflection.h
@@ -85,9 +85,6 @@ private:
 	/** Restore scene capture visibility on registered mesh actors primitives. */
 	void RestorePrimitiveVisibilityState() const;
 
-	/** Update scene capture hidden components. */
-	void UpdateSceneCaptureComponents(USceneCaptureComponent2D& SceneCaptureComponent) const;
-
 	/** Find or create scene captures with CG and without CG. */
 	TStaticArray<TWeakObjectPtr<UCompositeShadowReflectionCatcherComponent>, 2> FindOrCreateSceneCapturePair(ACompositeActor& InOuter) const;
 
diff --git a/Engine/Plugins/Compositing/Composite/Source/Composite/Public/Layers/CompositeLayerSingleLightShadow.h b/Engine/Plugins/Compositing/Composite/Source/Composite/Public/Layers/CompositeLayerSingleLightShadow.h
--- a/Engine/Plugins/Compositing/Composite/Source/Composite/Public/Layers/CompositeLayerSingleLightShadow.h
+++ b/Engine/Plugins/Compositing/Composite/Source/Composite/Public/Layers/CompositeLayerSingleLightShadow.h
@@ -80,20 +80,20 @@ public:
 	float OrthographicWidth = DEFAULT_ORTHOWIDTH;
 
 	/**
-	* Resolution of the shadow map texture. Lower resolutions will cause more blurry shadows, but require a higher shadow bias.
+	* Resolution of the shadow map texture. Lower resolutions will cause more blurry shadows.
 	*/
-	UPROPERTY(EditAnywhere, BlueprintReadWrite, Category = "", meta = (ClampMin = "128", UIMin = "128"))
+	UPROPERTY(EditAnywhere, BlueprintReadWrite, Category = "", meta = (ClampMin = "128", UIMin = "128", UIMax="4096"))
 	int32 ShadowMapResolution = 2048;
 
 	/** Basic shadow strength multiplier, from 0 to 1. */
 	UPROPERTY(EditAnywhere, BlueprintReadWrite, Category = "", meta = (ClampMin = "0.0", ClampMax = "1.0", UIMin = "0.0", UIMax = "1.0"))
 	float ShadowStrength = 0.5f;
 
-	/**
-	* Shadow bias to avoid shadow acne. Increase until unwanted shadowing disappears (automatically scaled by r.Shadow.CSMDepthBias console variable).
-	*/
-	UPROPERTY(EditAnywhere, BlueprintReadWrite, Category = "", meta = (ClampMin = "0.0", UIMin = "0.0"))
-	float ShadowBias = 0.5f;
+#if WITH_EDITORONLY_DATA
+	UE_DEPRECATED(5.8, "Shadow Bias is now deprecated.")
+	UPROPERTY(meta = (DeprecatedProperty, DeprecationMessage = "Shadow Bias is now deprecated."))
+	float ShadowBias_DEPRECATED = 0.5f;
+#endif
 
 private:
 	/** Find or create scene depth capture component, correctly configured. */
@@ -102,17 +102,11 @@ private:
 	/** Find or create shadow map capture component, correctly configured. */
 	UCompositeSceneCapture2DComponent* FindOrCreateShadowMapCapture(ACompositeActor& InOuter) const;
 
-	/** Update scene capture hidden components. */
-	void UpdateSceneCaptureComponents(USceneCaptureComponent2D& SceneCaptureComponent) const;
-
 	/** Convenience function to calculate shadow matrices. */
 	void GetShadowMatrices(UCompositeSceneCapture2DComponent& ShadowMapCapture, FMatrix44f& OutShadowMatrix, FVector4f& OutShadowInvDeviceZToWorldZ) const;
 
 private:
-	/**
-	* Optional: List of shadow casting actors to be removed from the main camera view's scene depth.
-	* This helps create a cleaner shadow matte preventing aliased dark edges around CG objects.
-	*/
+	/** List of shadow casting actors (required). */
 	UPROPERTY(EditAnywhere, BlueprintGetter = GetShadowCastingActors, BlueprintSetter = SetShadowCastingActors, Category = "", meta = (AllowPrivateAccess = true))
 	TArray<TSoftObjectPtr<AActor>> ShadowCastingActors;
 
