Composure: Add tooltip support to actor pickers.

#Synced-CL 48399844

diff --git a/Engine/Plugins/Compositing/Composite/Source/CompositeEditor/Private/Customizations/CompositeLayerPlateCustomization.cpp b/Engine/Plugins/Compositing/Composite/Source/CompositeEditor/Private/Customizations/CompositeLayerPlateCustomization.cpp
--- a/Engine/Plugins/Compositing/Composite/Source/CompositeEditor/Private/Customizations/CompositeLayerPlateCustomization.cpp
+++ b/Engine/Plugins/Compositing/Composite/Source/CompositeEditor/Private/Customizations/CompositeLayerPlateCustomization.cpp
@@ -92,6 +92,7 @@ void FCompositeLayerPlateCustomization::CustomizeDetails(IDetailLayoutBuilder& D
 		IDetailGroup& ActorListGroup = PlateCategory.AddGroup("CompositeMeshContent", LOCTEXT("CompositeMeshGroupName", "Composite Mesh Content"), false, true);
 		FCompositeActorPickerListRef ActorListRef(Plate, GET_MEMBER_NAME_CHECKED(UCompositeLayerPlate, CompositeMeshes), &Plate->CompositeMeshes);
 		
+		ActorListGroup.SetToolTip(ActorListRef.GetToolTipText());
 		ActorListGroup.AddWidgetRow()
 		.WholeRowContent()
 		[
diff --git a/Engine/Plugins/Compositing/Composite/Source/CompositeEditor/Private/Customizations/CompositeLayerShadowReflectionCustomization.cpp b/Engine/Plugins/Compositing/Composite/Source/CompositeEditor/Private/Customizations/CompositeLayerShadowReflectionCustomization.cpp
--- a/Engine/Plugins/Compositing/Composite/Source/CompositeEditor/Private/Customizations/CompositeLayerShadowReflectionCustomization.cpp
+++ b/Engine/Plugins/Compositing/Composite/Source/CompositeEditor/Private/Customizations/CompositeLayerShadowReflectionCustomization.cpp
@@ -53,6 +53,7 @@ void FCompositeLayerShadowReflectionCustomization::CustomizeDetails(IDetailLayou
 		IDetailGroup& ActorListGroup = LayerCategory.AddGroup("ShadowReflectionCatcherContent", LOCTEXT("ActorListGroupName", "Shadow/Reflection Catcher Content"), false, true);
 		FCompositeActorPickerListRef ActorListRef(ShadowReflection, GET_MEMBER_NAME_CHECKED(UCompositeLayerShadowReflection, Actors), &ShadowReflection->Actors);
 		
+		ActorListGroup.SetToolTip(ActorListRef.GetToolTipText());
 		ActorListGroup.AddWidgetRow()
 		.WholeRowContent()
 		[
diff --git a/Engine/Plugins/Compositing/Composite/Source/CompositeEditor/Private/Customizations/CompositeLayerSingleLightShadowCustomization.cpp b/Engine/Plugins/Compositing/Composite/Source/CompositeEditor/Private/Customizations/CompositeLayerSingleLightShadowCustomization.cpp
--- a/Engine/Plugins/Compositing/Composite/Source/CompositeEditor/Private/Customizations/CompositeLayerSingleLightShadowCustomization.cpp
+++ b/Engine/Plugins/Compositing/Composite/Source/CompositeEditor/Private/Customizations/CompositeLayerSingleLightShadowCustomization.cpp
@@ -50,6 +50,7 @@ void FCompositeLayerSingleLightShadowCustomization::CustomizeDetails(IDetailLayo
 		IDetailGroup& ActorListGroup = LayerCategory.AddGroup("ActorContent", LOCTEXT("ActorListGroupName", "Shadow Casting Content"), false, true);
 		FCompositeActorPickerListRef ActorListRef(SingleLightShadow, GET_MEMBER_NAME_CHECKED(UCompositeLayerSingleLightShadow, ShadowCastingActors), &SingleLightShadow->ShadowCastingActors);
 		
+		ActorListGroup.SetToolTip(ActorListRef.GetToolTipText());
 		ActorListGroup.AddWidgetRow()
 		.WholeRowContent()
 		[
diff --git a/Engine/Plugins/Compositing/Composite/Source/CompositeEditor/Private/UI/SCompositeActorPickerTable.cpp b/Engine/Plugins/Compositing/Composite/Source/CompositeEditor/Private/UI/SCompositeActorPickerTable.cpp
--- a/Engine/Plugins/Compositing/Composite/Source/CompositeEditor/Private/UI/SCompositeActorPickerTable.cpp
+++ b/Engine/Plugins/Compositing/Composite/Source/CompositeEditor/Private/UI/SCompositeActorPickerTable.cpp
@@ -374,6 +374,11 @@ void FCompositeActorPickerListRef::NotifyPostEditChangeList(EPropertyChangeType:
 	}
 }
 
+FText FCompositeActorPickerListRef::GetToolTipText(bool bShortTooltip) const
+{
+	return ActorListProperty ? ActorListProperty->GetToolTipText(bShortTooltip) : FText::GetEmpty();
+}
+
 void SCompositeActorPickerTable::Construct(const FArguments& InArgs, const FCompositeActorPickerListRef& InActorListRef)
 {
 	ActorListRef = InActorListRef;
diff --git a/Engine/Plugins/Compositing/Composite/Source/CompositeEditor/Private/UI/SCompositeActorPickerTable.h b/Engine/Plugins/Compositing/Composite/Source/CompositeEditor/Private/UI/SCompositeActorPickerTable.h
--- a/Engine/Plugins/Compositing/Composite/Source/CompositeEditor/Private/UI/SCompositeActorPickerTable.h
+++ b/Engine/Plugins/Compositing/Composite/Source/CompositeEditor/Private/UI/SCompositeActorPickerTable.h
@@ -27,6 +27,9 @@ public:
 
 	/** Raises the owner's PostEditchange event for the actor list property */
 	void NotifyPostEditChangeList(EPropertyChangeType::Type InChangeType);
+
+	/** Finds the localized tooltip or native tooltip as a fallback. */
+	FText GetToolTipText(bool bShortTooltip = false) const;
 	
 public:
 	/** Reference to the UObject that owns the actor list being edited by this widget */
