Composure: Remove async compute from denoising pass, which was missed when removed from the main keyer pass.

#Synced-CL 48399844

diff --git a/Engine/Plugins/Compositing/Composite/Source/Composite/Private/Passes/CompositePassColorKeyer.cpp b/Engine/Plugins/Compositing/Composite/Source/Composite/Private/Passes/CompositePassColorKeyer.cpp
--- a/Engine/Plugins/Compositing/Composite/Source/Composite/Private/Passes/CompositePassColorKeyer.cpp
+++ b/Engine/Plugins/Compositing/Composite/Source/Composite/Private/Passes/CompositePassColorKeyer.cpp
@@ -105,7 +105,7 @@ namespace UE
 				FComputeShaderUtils::AddPass(
 					GraphBuilder,
 					RDG_EVENT_NAME("CompositeCore.Denoise (%dx%d)", TextureSize.X, TextureSize.Y),
-					GSupportsEfficientAsyncCompute ? ERDGPassFlags::AsyncCompute : ERDGPassFlags::Compute,
+					ERDGPassFlags::Compute,
 					ComputeShader,
 					PassParameters,
 					FComputeShaderUtils::GetGroupCount(TextureSize, FCompositePassDenoiseShader::ThreadGroupSize)
