Composure: Renamed UCompositePassColorGrade to UCompositePassColorGrading, and added customization that displays pass enabled flag with specific type name (e.g. 'Enable Color Grading Pass')

#Synced-CL 48367791

diff --git a/Engine/Plugins/Compositing/Composite/Config/DefaultComposite.ini b/Engine/Plugins/Compositing/Composite/Config/DefaultComposite.ini
--- a/Engine/Plugins/Compositing/Composite/Config/DefaultComposite.ini
+++ b/Engine/Plugins/Compositing/Composite/Config/DefaultComposite.ini
@@ -40,3 +40,4 @@
 ; 5.8
 
 +PropertyRedirects=(OldName="/Script/Composite.CompositeViewProjectionComponent.MatrixParameterName", NewName="/Script/Composite.CompositeViewProjectionComponent.ViewProjectionMatrixParameter")
++ClassRedirects=(OldName="/Script/Composite.CompositePassColorGrade",NewName="/Script/Composite.CompositePassColorGrading")
diff --git a/Engine/Plugins/Compositing/Composite/Source/Composite/Private/Passes/CompositePassColorGrade.cpp b/Engine/Plugins/Compositing/Composite/Source/Composite/Private/Passes/CompositePassColorGrade.cpp
deleted file mode 100644
diff --git a/Engine/Plugins/Compositing/Composite/Source/Composite/Private/Passes/CompositePassColorGrading.cpp b/Engine/Plugins/Compositing/Composite/Source/Composite/Private/Passes/CompositePassColorGrading.cpp
new file mode 100644
--- /dev/null
+++ b/Engine/Plugins/Compositing/Composite/Source/Composite/Private/Passes/CompositePassColorGrading.cpp
@@ -0,0 +1,135 @@
+// Copyright Epic Games, Inc. All Rights Reserved.
+
+#include "Passes/CompositePassColorGrading.h"
+
+#include "SceneManagement.h"
+#include "ScreenPass.h"
+#include "PostProcess/PostProcessMaterialInputs.h"
+#include "Passes/CompositeCorePassProxy.h"
+
+DECLARE_GPU_STAT_NAMED(FCompositeColorGrade, TEXT("Composite.ColorGrade"));
+
+BEGIN_SHADER_PARAMETER_STRUCT(FColorGradePerRangeParameters, )
+	SHADER_PARAMETER(FVector4f, Saturation)
+	SHADER_PARAMETER(FVector4f, Contrast)
+	SHADER_PARAMETER(FVector4f, Gamma)
+	SHADER_PARAMETER(FVector4f, Gain)
+	SHADER_PARAMETER(FVector4f, Offset)
+END_SHADER_PARAMETER_STRUCT()
+
+class FCompositeColorGradeShader : public FGlobalShader
+{
+	DECLARE_GLOBAL_SHADER(FCompositeColorGradeShader);
+	SHADER_USE_PARAMETER_STRUCT(FCompositeColorGradeShader, FGlobalShader);
+
+	BEGIN_SHADER_PARAMETER_STRUCT(FParameters, )
+		SHADER_PARAMETER_RDG_TEXTURE(Texture2D, InputTexture)
+		SHADER_PARAMETER_SAMPLER(SamplerState, InputSampler)
+		SHADER_PARAMETER_STRUCT_REF(FWorkingColorSpaceShaderParameters, WorkingColorSpace)
+		SHADER_PARAMETER_STRUCT(FColorGradePerRangeParameters, Global)
+		SHADER_PARAMETER_STRUCT(FColorGradePerRangeParameters, Shadows)
+		SHADER_PARAMETER_STRUCT(FColorGradePerRangeParameters, Midtones)
+		SHADER_PARAMETER_STRUCT(FColorGradePerRangeParameters, Highlights)
+		SHADER_PARAMETER(float, ShadowsMax)
+		SHADER_PARAMETER(float, HighlightsMin)
+		SHADER_PARAMETER(float, HighlightsMax)
+		SHADER_PARAMETER(uint32, bIsTemperatureWhiteBalance)
+		SHADER_PARAMETER(float, WhiteTemp)
+		SHADER_PARAMETER(float, WhiteTint)
+		RENDER_TARGET_BINDING_SLOTS()
+	END_SHADER_PARAMETER_STRUCT()
+
+	static FColorGradePerRangeParameters ToParameters(const FColorGradePerRangeSettings& InParams)
+	{
+		FColorGradePerRangeParameters Output;
+		Output.Saturation = FVector4f(InParams.Saturation);
+		Output.Contrast = FVector4f(InParams.Contrast);
+		Output.Gamma = FVector4f(InParams.Gamma);
+		Output.Gain = FVector4f(InParams.Gain);
+		Output.Offset = FVector4f(InParams.Offset);
+
+		return Output;
+	}
+};
+IMPLEMENT_GLOBAL_SHADER(FCompositeColorGradeShader, "/Plugin/Composite/Private/CompositeColorGrade.usf", "MainPS", SF_Pixel);
+
+class FColorGradeCompositePassProxy : public FCompositeCorePassProxy
+{
+public:
+	IMPLEMENT_COMPOSITE_PASS(FColorGradeCompositePassProxy);
+
+	FColorGradeCompositePassProxy(UE::CompositeCore::FPassInputDeclArray InPassDeclaredInputs)
+		: FCompositeCorePassProxy(MoveTemp(InPassDeclaredInputs))
+	{ }
+
+	UE::CompositeCore::FPassTexture Add(FRDGBuilder& GraphBuilder, const FSceneView& InView, const UE::CompositeCore::FPassInputArray& Inputs, const UE::CompositeCore::FPassContext& PassContext) const override
+	{
+		RDG_EVENT_SCOPE_STAT(GraphBuilder, FCompositeColorGrade, "Composite.ColorGrade");
+		RDG_GPU_STAT_SCOPE(GraphBuilder, FCompositeColorGrade);
+
+		check(ValidateInputs(Inputs));
+
+		const UE::CompositeCore::FResourceMetadata& Metadata = Inputs[0].Metadata;
+		const FScreenPassTexture& Input = Inputs[0].Texture;
+		FScreenPassRenderTarget Output = Inputs.OverrideOutput;
+
+		// If the override output is provided, it means that this is the last pass in post processing.
+		if (!Output.IsValid())
+		{
+			Output = CreateOutputRenderTarget(GraphBuilder, InView, PassContext.OutputViewRect, Input.Texture->Desc, TEXT("ColorGradeCompositePass"));
+		}
+
+		FGlobalShaderMap* GlobalShaderMap = GetGlobalShaderMap(InView.GetFeatureLevel());
+		FScreenPassTextureViewport Viewport = FScreenPassTextureViewport(Output);
+
+		FCompositeColorGradeShader::FParameters* Parameters = GraphBuilder.AllocParameters<FCompositeColorGradeShader::FParameters>();
+		Parameters->InputTexture = Input.Texture;
+		Parameters->InputSampler = TStaticSamplerState<SF_Bilinear>::GetRHI();
+		Parameters->WorkingColorSpace = GDefaultWorkingColorSpaceUniformBuffer.GetUniformBufferRef();
+		Parameters->Global = FCompositeColorGradeShader::ToParameters(ColorGradingSettings.Global);
+		Parameters->Shadows = FCompositeColorGradeShader::ToParameters(ColorGradingSettings.Shadows);
+		Parameters->Midtones = FCompositeColorGradeShader::ToParameters(ColorGradingSettings.Midtones);
+		Parameters->Highlights = FCompositeColorGradeShader::ToParameters(ColorGradingSettings.Highlights);
+		Parameters->ShadowsMax = ColorGradingSettings.ShadowsMax;
+		Parameters->HighlightsMin = ColorGradingSettings.HighlightsMin;
+		Parameters->HighlightsMax = ColorGradingSettings.HighlightsMax;
+		Parameters->bIsTemperatureWhiteBalance = TemperatureSettings.TemperatureType;
+		Parameters->WhiteTemp = TemperatureSettings.WhiteTemp;
+		Parameters->WhiteTint = TemperatureSettings.WhiteTint;
+		Parameters->RenderTargets[0] = Output.GetRenderTargetBinding();
+
+		TShaderMapRef<FCompositeColorGradeShader> PixelShader(GlobalShaderMap);
+		AddDrawScreenPass(
+			GraphBuilder,
+			RDG_EVENT_NAME("Composite.ColorGrade (%dx%d)", Viewport.Extent.X, Viewport.Extent.Y),
+			InView,
+			Viewport,
+			Viewport,
+			PixelShader,
+			Parameters
+		);
+
+		return UE::CompositeCore::FPassTexture{ MoveTemp(Output), Metadata };
+	}
+
+	FCompositeTemperatureSettings TemperatureSettings;
+	FColorGradingSettings ColorGradingSettings;
+};
+
+UCompositePassColorGrading::UCompositePassColorGrading(const FObjectInitializer& ObjectInitializer)
+	: Super(ObjectInitializer)
+{
+}
+
+UCompositePassColorGrading::~UCompositePassColorGrading() = default;
+
+bool UCompositePassColorGrading::GetProxy(const UE::CompositeCore::FPassInputDecl& InputDecl, FSceneRenderingBulkObjectAllocator& InFrameAllocator, FCompositeCorePassProxy*& OutProxy) const
+{
+	FColorGradeCompositePassProxy* Proxy = InFrameAllocator.Create<FColorGradeCompositePassProxy>(UE::CompositeCore::FPassInputDeclArray{ InputDecl });
+	Proxy->TemperatureSettings = TemperatureSettings;
+	Proxy->ColorGradingSettings = ColorGradingSettings;
+
+	OutProxy = Proxy;
+	return true;
+}
+
diff --git a/Engine/Plugins/Compositing/Composite/Source/Composite/Public/Passes/CompositePassBase.h b/Engine/Plugins/Compositing/Composite/Source/Composite/Public/Passes/CompositePassBase.h
--- a/Engine/Plugins/Compositing/Composite/Source/Composite/Public/Passes/CompositePassBase.h
+++ b/Engine/Plugins/Compositing/Composite/Source/Composite/Public/Passes/CompositePassBase.h
@@ -64,6 +64,8 @@ protected:
 	UPROPERTY(EditAnywhere, Category = "", meta = (DisplayName="Name", DisplayPriority = "0"))
 	FString DisplayName;
 #endif
+
+	friend class FCompositePassBaseCustomization;
 };
 
 #undef UE_API
diff --git a/Engine/Plugins/Compositing/Composite/Source/Composite/Public/Passes/CompositePassColorGrade.h b/Engine/Plugins/Compositing/Composite/Source/Composite/Public/Passes/CompositePassColorGrade.h
deleted file mode 100644
diff --git a/Engine/Plugins/Compositing/Composite/Source/Composite/Public/Passes/CompositePassColorGrading.h b/Engine/Plugins/Compositing/Composite/Source/Composite/Public/Passes/CompositePassColorGrading.h
new file mode 100644
--- /dev/null
+++ b/Engine/Plugins/Compositing/Composite/Source/Composite/Public/Passes/CompositePassColorGrading.h
@@ -0,0 +1,64 @@
+// Copyright Epic Games, Inc. All Rights Reserved.
+
+#pragma once
+
+#include "CoreTypes.h"
+#include "CompositePassBase.h"
+#include "Passes/CompositeCorePassMergeProxy.h"
+
+#include "CompositePassColorGrading.generated.h"
+
+#define UE_API COMPOSITE_API
+
+/** Color grading pass temperature settings. */
+USTRUCT(BlueprintType)
+struct FCompositeTemperatureSettings
+{
+	GENERATED_USTRUCT_BODY()
+
+	/**
+	* Selects the type of temperature calculation.
+	* White Balance uses the Temperature value to control the virtual camera's White Balance. This is the default selection.
+	* Color Temperature uses the Temperature value to adjust the color temperature of the scene, which is the inverse of the White Balance operation.
+	*/
+	UPROPERTY(interp, BlueprintReadWrite, Category = "Color Grading|Temperature", meta = (DisplayName = "Temperature Type"))
+	TEnumAsByte<enum ETemperatureMethod> TemperatureType = ETemperatureMethod::TEMP_WhiteBalance;
+
+	/** Controls the color temperature or white balance in degrees Kelvin which the scene considers as white light. */
+	UPROPERTY(interp, BlueprintReadWrite, Category = "Color Grading|Temperature", meta = (UIMin = "1500.0", UIMax = "15000.0", DisplayName = "Temp"))
+	float WhiteTemp = 6500.0f;
+	
+	/** Controls the color of the scene along the magenta - green axis (orthogonal to the color temperature).  This feature is equivalent to color tint in digital cameras. */
+	UPROPERTY(interp, BlueprintReadWrite, Category = "Color Grading|Temperature", meta = (UIMin = "-1.0", UIMax = "1.0", DisplayName = "Tint"))
+	float WhiteTint = 0.0f;
+
+};
+
+/**
+ * Color grade pass.
+ * Assumes input is in linear working color space.
+*/
+UCLASS(MinimalAPI, BlueprintType, Blueprintable, EditInlineNew, CollapseCategories, meta = (DisplayName = "Color Grading Pass"))
+class UCompositePassColorGrading : public UCompositePassBase
+{
+	GENERATED_BODY()
+
+public:
+	/** Constructor */
+	UE_API UCompositePassColorGrading(const FObjectInitializer& ObjectInitializer);
+	/** Destructor */
+	UE_API ~UCompositePassColorGrading();
+
+	UE_API virtual bool GetProxy(const UE::CompositeCore::FPassInputDecl& InputDecl, FSceneRenderingBulkObjectAllocator& InFrameAllocator, FCompositeCorePassProxy*& OutProxy) const override;
+
+public:
+	/** Color temperature settings. */
+	UPROPERTY(BlueprintReadWrite, EditAnywhere, Category = "", meta = (DisplayName = "Temperature"))
+	FCompositeTemperatureSettings TemperatureSettings;
+
+	/** Color grade settings. */
+	UPROPERTY(BlueprintReadWrite, EditAnywhere, Category = "", meta = (ShowOnlyInnerProperties))
+	FColorGradingSettings ColorGradingSettings;
+};
+
+#undef UE_API
diff --git a/Engine/Plugins/Compositing/Composite/Source/CompositeEditor/Private/Customizations/CompositePassBaseCustomization.cpp b/Engine/Plugins/Compositing/Composite/Source/CompositeEditor/Private/Customizations/CompositePassBaseCustomization.cpp
new file mode 100644
--- /dev/null
+++ b/Engine/Plugins/Compositing/Composite/Source/CompositeEditor/Private/Customizations/CompositePassBaseCustomization.cpp
@@ -0,0 +1,38 @@
+﻿// Copyright Epic Games, Inc. All Rights Reserved.
+
+#include "CompositePassBaseCustomization.h"
+
+#include "DetailLayoutBuilder.h"
+#include "Passes/CompositePassBase.h"
+
+#define LOCTEXT_NAMESPACE "FCompositePassBaseCustomization"
+
+TSharedRef<IDetailCustomization> FCompositePassBaseCustomization::MakeInstance()
+{
+	return MakeShared<FCompositePassBaseCustomization>();
+}
+
+void FCompositePassBaseCustomization::CustomizeDetails(IDetailLayoutBuilder& DetailLayout)
+{
+	TArray<TWeakObjectPtr<UObject>> Objects;
+	DetailLayout.GetObjectsBeingCustomized(Objects);
+
+	// Rename the bIsEnabled property to specify the type of pass being enabled/disabled. If multiple passes are displayed at once, show "Enable All Passes"
+	FText ObjectTypeName = FText::GetEmpty();
+	if (Objects.Num() == 1)
+	{
+		if (TStrongObjectPtr<UObject> PinnedObject = Objects[0].Pin())
+		{
+			ObjectTypeName = PinnedObject->GetClass()->GetDisplayNameText();
+		}
+	}
+	else if (Objects.Num() > 1)
+	{
+		ObjectTypeName = LOCTEXT("AllPassesLabel", "All Passes");
+	}
+	
+	TSharedPtr<IPropertyHandle> EnabledHandle = DetailLayout.GetProperty(GET_MEMBER_NAME_CHECKED(UCompositePassBase, bIsEnabled));
+	DetailLayout.EditDefaultProperty(EnabledHandle)->DisplayName(FText::Format(LOCTEXT("PassEnabledLabelFormat", "Enable {0}"), ObjectTypeName));
+}
+
+#undef LOCTEXT_NAMESPACE
\ No newline at end of file
diff --git a/Engine/Plugins/Compositing/Composite/Source/CompositeEditor/Private/Customizations/CompositePassBaseCustomization.h b/Engine/Plugins/Compositing/Composite/Source/CompositeEditor/Private/Customizations/CompositePassBaseCustomization.h
new file mode 100644
--- /dev/null
+++ b/Engine/Plugins/Compositing/Composite/Source/CompositeEditor/Private/Customizations/CompositePassBaseCustomization.h
@@ -0,0 +1,16 @@
+﻿// Copyright Epic Games, Inc. All Rights Reserved.
+
+#pragma once
+
+#include "IDetailCustomization.h"
+
+/**
+ * Base customization for all derivations of UCompositePassBase
+ */
+class FCompositePassBaseCustomization : public IDetailCustomization
+{
+public:
+	static TSharedRef<IDetailCustomization> MakeInstance();
+	
+	virtual void CustomizeDetails(IDetailLayoutBuilder& DetailLayout) override;
+};
diff --git a/Engine/Plugins/Compositing/Composite/Source/CompositeEditor/Private/UI/SCompositePlatePassPanel.cpp b/Engine/Plugins/Compositing/Composite/Source/CompositeEditor/Private/UI/SCompositePlatePassPanel.cpp
--- a/Engine/Plugins/Compositing/Composite/Source/CompositeEditor/Private/UI/SCompositePlatePassPanel.cpp
+++ b/Engine/Plugins/Compositing/Composite/Source/CompositeEditor/Private/UI/SCompositePlatePassPanel.cpp
@@ -7,6 +7,7 @@
 #include "PropertyEditorModule.h"
 #include "PropertyPath.h"
 #include "SCompositePassTree.h"
+#include "Customizations/CompositePassBaseCustomization.h"
 #include "Layers/CompositeLayerPlate.h"
 #include "Modules/ModuleManager.h"
 
@@ -24,8 +25,10 @@ void SCompositePlatePassPanel::Construct(const FArguments& InArgs, UCompositeLay
 	
 	DetailsView = FModuleManager::GetModuleChecked<FPropertyEditorModule>("PropertyEditor").CreateDetailView(DetailsViewArgs);
 
-	DetailsView->CountRows();
-	
+	DetailsView->RegisterInstancedCustomPropertyLayout(
+		UCompositePassBase::StaticClass(),
+		FOnGetDetailCustomizationInstance::CreateStatic(&FCompositePassBaseCustomization::MakeInstance));
+
 	ChildSlot
 	[
 		SNew(SVerticalBox)
