Composite: Color grading fixes & improvements.

* Fixed incorrectly flipped WhiteBalance setting.
* New option to unpremultiply for color grading, enabled by default.
* Renamed shader & parameters with "color grading".

#Synced-CL 48523035

diff --git a/Engine/Plugins/Compositing/Composite/Shaders/Private/CompositeColorGrade.usf b/Engine/Plugins/Compositing/Composite/Shaders/Private/CompositeColorGrade.usf
deleted file mode 100644
diff --git a/Engine/Plugins/Compositing/Composite/Shaders/Private/CompositeColorGrading.usf b/Engine/Plugins/Compositing/Composite/Shaders/Private/CompositeColorGrading.usf
new file mode 100644
--- /dev/null
+++ b/Engine/Plugins/Compositing/Composite/Shaders/Private/CompositeColorGrading.usf
@@ -0,0 +1,132 @@
+// Copyright Epic Games, Inc. All Rights Reserved.
+
+#include "/Engine/Private/Common.ush"
+#include "/Engine/Private/TonemapCommon.ush"
+#include "/Engine/Private/ACES/ACESCommon.ush"
+
+#include "CompositeCommon.ush"
+
+static const float DefaultTemperature = 6500.0;
+static const float DefaultTint = 0.0;
+static const float Epsilon = 1e-8;
+
+Texture2D InputTexture;
+SamplerState InputSampler;
+
+uint bIsTemperatureWhiteBalance;
+float WhiteTemp;
+float WhiteTint;
+
+float4 Global_Saturation;
+float4 Global_Contrast;
+float4 Global_Gamma;
+float4 Global_Gain;
+float4 Global_Offset;
+
+float4 Shadows_Saturation;
+float4 Shadows_Contrast;
+float4 Shadows_Gamma;
+float4 Shadows_Gain;
+float4 Shadows_Offset;
+
+float4 Midtones_Saturation;
+float4 Midtones_Contrast;
+float4 Midtones_Gamma;
+float4 Midtones_Gain;
+float4 Midtones_Offset;
+
+float4 Highlights_Saturation;
+float4 Highlights_Contrast;
+float4 Highlights_Gamma;
+float4 Highlights_Gain;
+float4 Highlights_Offset;
+
+uint bUnpremult;
+float ShadowsMax;
+float HighlightsMin;
+float HighlightsMax;
+
+float3 ColorCorrect( float3 WorkingColor,
+	float4 ColorSaturation,
+	float4 ColorContrast,
+	float4 ColorGamma,
+	float4 ColorGain,
+	float4 ColorOffset )
+{
+	float Luma = dot( WorkingColor, AP1_RGB2Y );
+	WorkingColor = max( 0, lerp( Luma.xxx, WorkingColor, ColorSaturation.xyz*ColorSaturation.w ) );
+	WorkingColor = pow( WorkingColor * (1.0 / 0.18), ColorContrast.xyz*ColorContrast.w ) * 0.18;
+	WorkingColor = pow( WorkingColor, 1.0 / (ColorGamma.xyz*ColorGamma.w) );
+	WorkingColor = WorkingColor * (ColorGain.xyz * ColorGain.w) + (ColorOffset.xyz + ColorOffset.w);
+	return WorkingColor;
+}
+
+// Nuke-style Color Correct
+float3 ColorCorrectAll( float3 WorkingColor )
+{
+	float Luma = dot( WorkingColor, AP1_RGB2Y );
+
+	// Shadow CC
+	float3 CCColorShadows = ColorCorrect(WorkingColor, 
+		Shadows_Saturation*Global_Saturation, 
+		Shadows_Contrast*Global_Contrast, 
+		Shadows_Gamma*Global_Gamma, 
+		Shadows_Gain*Global_Gain, 
+		Shadows_Offset+Global_Offset);
+	float CCWeightShadows = 1- smoothstep(0, ShadowsMax, Luma);
+	
+	// Highlight CC
+	float3 CCColorHighlights = ColorCorrect(WorkingColor, 
+		Highlights_Saturation*Global_Saturation, 
+		Highlights_Contrast*Global_Contrast, 
+		Highlights_Gamma*Global_Gamma, 
+		Highlights_Gain*Global_Gain, 
+		Highlights_Offset+Global_Offset);
+	float CCWeightHighlights = smoothstep(HighlightsMin, HighlightsMax, Luma);
+
+	// Midtone CC
+	float3 CCColorMidtones = ColorCorrect(WorkingColor, 
+		Midtones_Saturation*Global_Saturation, 
+		Midtones_Contrast*Global_Contrast, 
+		Midtones_Gamma*Global_Gamma, 
+		Midtones_Gain*Global_Gain, 
+		Midtones_Offset+Global_Offset);
+	float CCWeightMidtones = 1 - CCWeightShadows - CCWeightHighlights;
+
+	// Blend Shadow, Midtone and Highlight CCs
+	float3 WorkingColorSMH = CCColorShadows*CCWeightShadows + CCColorMidtones*CCWeightMidtones + CCColorHighlights*CCWeightHighlights;
+	
+	return WorkingColorSMH;
+}
+
+void MainPS(
+    noperspective float4 InUVAndScreenPos : TEXCOORD0,
+    out float4 OutColor : SV_Target0
+    )
+{
+    float4 LinearColor = InputTexture.Sample(InputSampler, InUVAndScreenPos.xy);
+	
+	if(bUnpremult)
+	{
+		LinearColor = Unpremultiply(LinearColor);
+	}
+    
+    BRANCH
+	if( abs(WhiteTemp - DefaultTemperature) > Epsilon || abs(WhiteTint - DefaultTint) > Epsilon )
+	{
+		LinearColor.rgb = WhiteBalance(LinearColor.rgb, WhiteTemp, WhiteTint, bIsTemperatureWhiteBalance, (float3x3) WorkingColorSpace.ToXYZ, (float3x3) WorkingColorSpace.FromXYZ);
+	}
+
+    float4 ColorAP1 = mul(WorkingColorSpace.ToAP1, LinearColor);
+    
+    ColorAP1.rgb = ColorCorrectAll(ColorAP1.rgb);
+
+	LinearColor = mul(WorkingColorSpace.FromAP1, ColorAP1);
+	
+	if(bUnpremult)
+	{
+		LinearColor = Premultiply(LinearColor);
+	}
+
+	OutColor = LinearColor;
+}
diff --git a/Engine/Plugins/Compositing/Composite/Shaders/Private/CompositeColorKeyerCommon.ush b/Engine/Plugins/Compositing/Composite/Shaders/Private/CompositeColorKeyerCommon.ush
--- a/Engine/Plugins/Compositing/Composite/Shaders/Private/CompositeColorKeyerCommon.ush
+++ b/Engine/Plugins/Compositing/Composite/Shaders/Private/CompositeColorKeyerCommon.ush
@@ -7,6 +7,8 @@
 #include "/Engine/Private/GammaCorrectionCommon.ush"
 #include "/Engine/Private/PostProcessCommon.ush"
 
+#include "CompositeCommon.ush"
+
 // CIE 1931 D65 white point
 static const float3 XYZ_n = xyY_2_XYZ(float3(0.3127, 0.3290, 1.0));
 static const float ZeroEpsilon = 1.0e-6;
@@ -14,22 +16,6 @@ static const float DeltaE_JND = 1.0;
 static const float Lstar_CIELUV = 100.0;
 static const float Lstar_CIELUV_Inv = 1.0 / Lstar_CIELUV;
 
-// Safe reciprocal
-float SafeRcp(float x)
-{
-	return x > 0.0 ? rcp(x) : 0.0;
-}
-
-float4 Unpremultiply(float4 InColor)
-{
-	return float4(InColor.xyz * SafeRcp(InColor.w), InColor.w);
-}
-
-float4 Premultiply(float4 InColor)
-{
-	return float4(InColor.xyz * InColor.w, InColor.w);
-}
-
 float LuminanceSRGB(float3 InColor)
 {
 	return dot(InColor, float3(0.2126390059, 0.7151686788, 0.0721923154));
diff --git a/Engine/Plugins/Compositing/Composite/Shaders/Private/CompositeCommon.ush b/Engine/Plugins/Compositing/Composite/Shaders/Private/CompositeCommon.ush
--- a/Engine/Plugins/Compositing/Composite/Shaders/Private/CompositeCommon.ush
+++ b/Engine/Plugins/Compositing/Composite/Shaders/Private/CompositeCommon.ush
@@ -6,6 +6,22 @@
 #define ESourceEncoding_Gamma 1
 #define ESourceEncoding_sRGB  2
 
+// Safe reciprocal
+float SafeRcp(float x)
+{
+	return x > 0.0 ? rcp(x) : 0.0;
+}
+
+float4 Unpremultiply(float4 InColor)
+{
+	return float4(InColor.xyz * SafeRcp(InColor.w), InColor.w);
+}
+
+float4 Premultiply(float4 InColor)
+{
+	return float4(InColor.xyz * InColor.w, InColor.w);
+}
+
 float3 Decode(float3 InColor, uint InEncoding)
 {
     BRANCH
diff --git a/Engine/Plugins/Compositing/Composite/Source/Composite/Private/Passes/CompositePassColorGrading.cpp b/Engine/Plugins/Compositing/Composite/Source/Composite/Private/Passes/CompositePassColorGrading.cpp
--- a/Engine/Plugins/Compositing/Composite/Source/Composite/Private/Passes/CompositePassColorGrading.cpp
+++ b/Engine/Plugins/Compositing/Composite/Source/Composite/Private/Passes/CompositePassColorGrading.cpp
@@ -7,9 +7,9 @@
 #include "PostProcess/PostProcessMaterialInputs.h"
 #include "Passes/CompositeCorePassProxy.h"
 
-DECLARE_GPU_STAT_NAMED(FCompositeColorGrade, TEXT("Composite.ColorGrade"));
+DECLARE_GPU_STAT_NAMED(FCompositeColorGrading, TEXT("Composite.ColorGrading"));
 
-BEGIN_SHADER_PARAMETER_STRUCT(FColorGradePerRangeParameters, )
+BEGIN_SHADER_PARAMETER_STRUCT(FColorGradingPerRangeParameters, )
 	SHADER_PARAMETER(FVector4f, Saturation)
 	SHADER_PARAMETER(FVector4f, Contrast)
 	SHADER_PARAMETER(FVector4f, Gamma)
@@ -17,19 +17,20 @@ BEGIN_SHADER_PARAMETER_STRUCT(FColorGradePerRangeParameters, )
 	SHADER_PARAMETER(FVector4f, Offset)
 END_SHADER_PARAMETER_STRUCT()
 
-class FCompositeColorGradeShader : public FGlobalShader
+class FCompositeColorGradingShader : public FGlobalShader
 {
-	DECLARE_GLOBAL_SHADER(FCompositeColorGradeShader);
-	SHADER_USE_PARAMETER_STRUCT(FCompositeColorGradeShader, FGlobalShader);
+	DECLARE_GLOBAL_SHADER(FCompositeColorGradingShader);
+	SHADER_USE_PARAMETER_STRUCT(FCompositeColorGradingShader, FGlobalShader);
 
 	BEGIN_SHADER_PARAMETER_STRUCT(FParameters, )
 		SHADER_PARAMETER_RDG_TEXTURE(Texture2D, InputTexture)
 		SHADER_PARAMETER_SAMPLER(SamplerState, InputSampler)
 		SHADER_PARAMETER_STRUCT_REF(FWorkingColorSpaceShaderParameters, WorkingColorSpace)
-		SHADER_PARAMETER_STRUCT(FColorGradePerRangeParameters, Global)
-		SHADER_PARAMETER_STRUCT(FColorGradePerRangeParameters, Shadows)
-		SHADER_PARAMETER_STRUCT(FColorGradePerRangeParameters, Midtones)
-		SHADER_PARAMETER_STRUCT(FColorGradePerRangeParameters, Highlights)
+		SHADER_PARAMETER_STRUCT(FColorGradingPerRangeParameters, Global)
+		SHADER_PARAMETER_STRUCT(FColorGradingPerRangeParameters, Shadows)
+		SHADER_PARAMETER_STRUCT(FColorGradingPerRangeParameters, Midtones)
+		SHADER_PARAMETER_STRUCT(FColorGradingPerRangeParameters, Highlights)
+		SHADER_PARAMETER(uint32, bUnpremult)
 		SHADER_PARAMETER(float, ShadowsMax)
 		SHADER_PARAMETER(float, HighlightsMin)
 		SHADER_PARAMETER(float, HighlightsMax)
@@ -39,9 +40,9 @@ class FCompositeColorGradeShader : public FGlobalShader
 		RENDER_TARGET_BINDING_SLOTS()
 	END_SHADER_PARAMETER_STRUCT()
 
-	static FColorGradePerRangeParameters ToParameters(const FColorGradePerRangeSettings& InParams)
+	static FColorGradingPerRangeParameters ToParameters(const FColorGradePerRangeSettings& InParams)
 	{
-		FColorGradePerRangeParameters Output;
+		FColorGradingPerRangeParameters Output;
 		Output.Saturation = FVector4f(InParams.Saturation);
 		Output.Contrast = FVector4f(InParams.Contrast);
 		Output.Gamma = FVector4f(InParams.Gamma);
@@ -51,21 +52,21 @@ class FCompositeColorGradeShader : public FGlobalShader
 		return Output;
 	}
 };
-IMPLEMENT_GLOBAL_SHADER(FCompositeColorGradeShader, "/Plugin/Composite/Private/CompositeColorGrade.usf", "MainPS", SF_Pixel);
+IMPLEMENT_GLOBAL_SHADER(FCompositeColorGradingShader, "/Plugin/Composite/Private/CompositeColorGrading.usf", "MainPS", SF_Pixel);
 
-class FColorGradeCompositePassProxy : public FCompositeCorePassProxy
+class FColorGradingCompositePassProxy : public FCompositeCorePassProxy
 {
 public:
-	IMPLEMENT_COMPOSITE_PASS(FColorGradeCompositePassProxy);
+	IMPLEMENT_COMPOSITE_PASS(FColorGradingCompositePassProxy);
 
-	FColorGradeCompositePassProxy(UE::CompositeCore::FPassInputDeclArray InPassDeclaredInputs)
+	FColorGradingCompositePassProxy(UE::CompositeCore::FPassInputDeclArray InPassDeclaredInputs)
 		: FCompositeCorePassProxy(MoveTemp(InPassDeclaredInputs))
 	{ }
 
 	UE::CompositeCore::FPassTexture Add(FRDGBuilder& GraphBuilder, const FSceneView& InView, const UE::CompositeCore::FPassInputArray& Inputs, const UE::CompositeCore::FPassContext& PassContext) const override
 	{
-		RDG_EVENT_SCOPE_STAT(GraphBuilder, FCompositeColorGrade, "Composite.ColorGrade");
-		RDG_GPU_STAT_SCOPE(GraphBuilder, FCompositeColorGrade);
+		RDG_EVENT_SCOPE_STAT(GraphBuilder, FCompositeColorGrading, "Composite.ColorGrading");
+		RDG_GPU_STAT_SCOPE(GraphBuilder, FCompositeColorGrading);
 
 		check(ValidateInputs(Inputs));
 
@@ -76,32 +77,33 @@ public:
 		// If the override output is provided, it means that this is the last pass in post processing.
 		if (!Output.IsValid())
 		{
-			Output = CreateOutputRenderTarget(GraphBuilder, InView, PassContext.OutputViewRect, Input.Texture->Desc, TEXT("ColorGradeCompositePass"));
+			Output = CreateOutputRenderTarget(GraphBuilder, InView, PassContext.OutputViewRect, Input.Texture->Desc, TEXT("ColorGradingCompositePass"));
 		}
 
 		FGlobalShaderMap* GlobalShaderMap = GetGlobalShaderMap(InView.GetFeatureLevel());
 		FScreenPassTextureViewport Viewport = FScreenPassTextureViewport(Output);
 
-		FCompositeColorGradeShader::FParameters* Parameters = GraphBuilder.AllocParameters<FCompositeColorGradeShader::FParameters>();
+		FCompositeColorGradingShader::FParameters* Parameters = GraphBuilder.AllocParameters<FCompositeColorGradingShader::FParameters>();
 		Parameters->InputTexture = Input.Texture;
 		Parameters->InputSampler = TStaticSamplerState<SF_Bilinear>::GetRHI();
 		Parameters->WorkingColorSpace = GDefaultWorkingColorSpaceUniformBuffer.GetUniformBufferRef();
-		Parameters->Global = FCompositeColorGradeShader::ToParameters(ColorGradingSettings.Global);
-		Parameters->Shadows = FCompositeColorGradeShader::ToParameters(ColorGradingSettings.Shadows);
-		Parameters->Midtones = FCompositeColorGradeShader::ToParameters(ColorGradingSettings.Midtones);
-		Parameters->Highlights = FCompositeColorGradeShader::ToParameters(ColorGradingSettings.Highlights);
+		Parameters->Global = FCompositeColorGradingShader::ToParameters(ColorGradingSettings.Global);
+		Parameters->Shadows = FCompositeColorGradingShader::ToParameters(ColorGradingSettings.Shadows);
+		Parameters->Midtones = FCompositeColorGradingShader::ToParameters(ColorGradingSettings.Midtones);
+		Parameters->Highlights = FCompositeColorGradingShader::ToParameters(ColorGradingSettings.Highlights);
+		Parameters->bUnpremult = bUnpremult;
 		Parameters->ShadowsMax = ColorGradingSettings.ShadowsMax;
 		Parameters->HighlightsMin = ColorGradingSettings.HighlightsMin;
 		Parameters->HighlightsMax = ColorGradingSettings.HighlightsMax;
-		Parameters->bIsTemperatureWhiteBalance = TemperatureSettings.TemperatureType;
+		Parameters->bIsTemperatureWhiteBalance = (TemperatureSettings.TemperatureType == ETemperatureMethod::TEMP_WhiteBalance);
 		Parameters->WhiteTemp = TemperatureSettings.WhiteTemp;
 		Parameters->WhiteTint = TemperatureSettings.WhiteTint;
 		Parameters->RenderTargets[0] = Output.GetRenderTargetBinding();
 
-		TShaderMapRef<FCompositeColorGradeShader> PixelShader(GlobalShaderMap);
+		TShaderMapRef<FCompositeColorGradingShader> PixelShader(GlobalShaderMap);
 		AddDrawScreenPass(
 			GraphBuilder,
-			RDG_EVENT_NAME("Composite.ColorGrade (%dx%d)", Viewport.Extent.X, Viewport.Extent.Y),
+			RDG_EVENT_NAME("Composite.ColorGrading (%dx%d)", Viewport.Extent.X, Viewport.Extent.Y),
 			InView,
 			Viewport,
 			Viewport,
@@ -114,10 +116,14 @@ public:
 
 	FCompositeTemperatureSettings TemperatureSettings;
 	FColorGradingSettings ColorGradingSettings;
+	bool bUnpremult = true;
 };
 
 UCompositePassColorGrading::UCompositePassColorGrading(const FObjectInitializer& ObjectInitializer)
 	: Super(ObjectInitializer)
+	, bUnpremult{true}
+	, TemperatureSettings{}
+	, ColorGradingSettings{}
 {
 }
 
@@ -125,9 +131,10 @@ UCompositePassColorGrading::~UCompositePassColorGrading() = default;
 
 bool UCompositePassColorGrading::GetProxy(const UE::CompositeCore::FPassInputDecl& InputDecl, FSceneRenderingBulkObjectAllocator& InFrameAllocator, FCompositeCorePassProxy*& OutProxy) const
 {
-	FColorGradeCompositePassProxy* Proxy = InFrameAllocator.Create<FColorGradeCompositePassProxy>(UE::CompositeCore::FPassInputDeclArray{ InputDecl });
+	FColorGradingCompositePassProxy* Proxy = InFrameAllocator.Create<FColorGradingCompositePassProxy>(UE::CompositeCore::FPassInputDeclArray{ InputDecl });
 	Proxy->TemperatureSettings = TemperatureSettings;
 	Proxy->ColorGradingSettings = ColorGradingSettings;
+	Proxy->bUnpremult = bUnpremult;
 
 	OutProxy = Proxy;
 	return true;
diff --git a/Engine/Plugins/Compositing/Composite/Source/Composite/Public/Passes/CompositePassColorGrading.h b/Engine/Plugins/Compositing/Composite/Source/Composite/Public/Passes/CompositePassColorGrading.h
--- a/Engine/Plugins/Compositing/Composite/Source/Composite/Public/Passes/CompositePassColorGrading.h
+++ b/Engine/Plugins/Compositing/Composite/Source/Composite/Public/Passes/CompositePassColorGrading.h
@@ -52,11 +52,15 @@ public:
 	UE_API virtual bool GetProxy(const UE::CompositeCore::FPassInputDecl& InputDecl, FSceneRenderingBulkObjectAllocator& InFrameAllocator, FCompositeCorePassProxy*& OutProxy) const override;
 
 public:
+	/** Color is divided by its alpha channel before processing, and multiplied again after. */
+	UPROPERTY(BlueprintReadWrite, EditAnywhere, Category = "")
+	bool bUnpremult;
+
 	/** Color temperature settings. */
 	UPROPERTY(BlueprintReadWrite, EditAnywhere, Category = "", meta = (DisplayName = "Temperature"))
 	FCompositeTemperatureSettings TemperatureSettings;
 
-	/** Color grade settings. */
+	/** Color grading settings. */
 	UPROPERTY(BlueprintReadWrite, EditAnywhere, Category = "", meta = (ShowOnlyInnerProperties))
 	FColorGradingSettings ColorGradingSettings;
 };
