Composure: Remove/deprecate AutoConfigureActors setting that controlled Hidden In Scene Capture for meshes on a shadow/reflection layer.

We previously kept meshes visible by default since shadowing and AO were lost when hidden. Upon further investigation, the default behavior can be corrected with the following cvars:
* r.MegaLights.ScreenTraces=0
* r.Lumen.ScreenProbeGather.ShortRangeAO.HardwareRayTracing=1

By doing so, we retain shadows & AO while removing undesirable edge artifacts that stemmed from keeping meshes visible. We accept the added cost given faster layers alternatives (single light shadow & future planar reflection layer).

#Synced-CL 48601581

diff --git a/Engine/Plugins/Compositing/Composite/Source/Composite/Private/Layers/CompositeLayerShadowReflection.cpp b/Engine/Plugins/Compositing/Composite/Source/Composite/Private/Layers/CompositeLayerShadowReflection.cpp
--- a/Engine/Plugins/Compositing/Composite/Source/Composite/Private/Layers/CompositeLayerShadowReflection.cpp
+++ b/Engine/Plugins/Compositing/Composite/Source/Composite/Private/Layers/CompositeLayerShadowReflection.cpp
@@ -5,6 +5,7 @@
 #include "CompositeActor.h"
 #include "Components/PrimitiveComponent.h"
 #include "Components/CompositeShadowReflectionCatcherComponent.h"
+#include "CompositeModule.h"
 #include "CompositeRenderTargetPool.h"
 #include "Containers/StaticArray.h"
 
@@ -12,7 +13,6 @@ UCompositeLayerShadowReflection::UCompositeLayerShadowReflection(const FObjectIn
 	: Super(ObjectInitializer)
 {
 	Operation = ECompositeCoreMergeOp::Multiply;
-	AutoConfigureActors = ECompositeHiddenInSceneCaptureConfiguration::None;
 
 	if (!HasAnyFlags(RF_ClassDefaultObject))
 	{
@@ -31,9 +31,16 @@ void UCompositeLayerShadowReflection::OnRemoved(ACompositeActor* LastOwner)
 		LastOwner->DestroySceneCaptures(this);
 	}
 
-	if (AutoConfigureActors != ECompositeHiddenInSceneCaptureConfiguration::None)
+	SetHiddenPrimitiveState(false);
+}
+
+void UCompositeLayerShadowReflection::PostLoad()
+{
+	Super::PostLoad();
+
+	if (IsEnabled() && !Actors.IsEmpty())
 	{
-		RestorePrimitiveVisibilityState();
+		SetHiddenPrimitiveState(true);
 	}
 }
 
@@ -58,10 +65,8 @@ void UCompositeLayerShadowReflection::PreEditChange(FProperty* PropertyAboutToCh
 
 	if (PropertyName == GET_MEMBER_NAME_CHECKED(ThisClass, Actors))
 	{
-		if (AutoConfigureActors != ECompositeHiddenInSceneCaptureConfiguration::None)
-		{
-			RestorePrimitiveVisibilityState();
-		}
+		// Pre-actively remove hidden state in case actors are removed
+		SetHiddenPrimitiveState(false);
 	}
 }
 
@@ -75,27 +80,14 @@ void UCompositeLayerShadowReflection::PostEditChangeProperty(FPropertyChangedEve
 	{
 		SetActors(MoveTemp(Actors));
 	}
-	else if (PropertyName == GET_MEMBER_NAME_CHECKED(ThisClass, AutoConfigureActors))
-	{
-		if (AutoConfigureActors != ECompositeHiddenInSceneCaptureConfiguration::None)
-		{
-			const bool bHiddenInSceneCapture = (AutoConfigureActors == ECompositeHiddenInSceneCaptureConfiguration::Hidden);
-			const bool bAffectIndirectLightingWhileHidden = bHiddenInSceneCapture;
-			const bool bCastHiddenShadow = bHiddenInSceneCapture;
-
-			UpdatePrimitiveVisibilityState(bHiddenInSceneCapture, bAffectIndirectLightingWhileHidden, bCastHiddenShadow);
-		}
-	}
 }
 
 void UCompositeLayerShadowReflection::PreEditUndo()
 {
 	Super::PreEditUndo();
 
-	if (AutoConfigureActors != ECompositeHiddenInSceneCaptureConfiguration::None)
-	{
-		RestorePrimitiveVisibilityState();
-	}
+	// Pre-actively remove hidden state in case undo empties actors
+	SetHiddenPrimitiveState(false);
 }
 
 void UCompositeLayerShadowReflection::PostEditUndo()
@@ -182,20 +174,14 @@ const TArray<TSoftObjectPtr<AActor>> UCompositeLayerShadowReflection::GetActors(
 
 void UCompositeLayerShadowReflection::SetActors(TArray<TSoftObjectPtr<AActor>> InActors)
 {
-	if (AutoConfigureActors != ECompositeHiddenInSceneCaptureConfiguration::None)
-	{
-		RestorePrimitiveVisibilityState();
-	}
+	// First, we make sure primitive hidden state is removed
+	SetHiddenPrimitiveState(false);
 
 	Actors = MoveTemp(InActors);
 
-	if (AutoConfigureActors != ECompositeHiddenInSceneCaptureConfiguration::None)
+	if (IsEnabled() && !Actors.IsEmpty())
 	{
-		const bool bHiddenInSceneCapture = (AutoConfigureActors == ECompositeHiddenInSceneCaptureConfiguration::Hidden);
-		const bool bAffectIndirectLightingWhileHidden = bHiddenInSceneCapture;
-		const bool bCastHiddenShadow = bHiddenInSceneCapture;
-
-		UpdatePrimitiveVisibilityState(bHiddenInSceneCapture, bAffectIndirectLightingWhileHidden, bCastHiddenShadow);
+		SetHiddenPrimitiveState(true);
 	}
 
 	ACompositeActor* CompositeActor = GetTypedOuter<ACompositeActor>();
@@ -221,6 +207,13 @@ void UCompositeLayerShadowReflection::SetActors(TArray<TSoftObjectPtr<AActor>> I
 	UpdateSceneCaptureComponents(*CachedSceneCaptures[1], Actors);
 }
 
+void UCompositeLayerShadowReflection::SetEnabled(bool bInEnabled)
+{
+	Super::SetEnabled(bInEnabled);
+
+	SetHiddenPrimitiveState(bInEnabled && !Actors.IsEmpty());
+}
+
 void UCompositeLayerShadowReflection::UpdatePrimitiveVisibilityState(
 	bool bInHiddenInSceneCapture,
 	bool bInAffectIndirectLightingWhileHidden,
@@ -264,12 +257,74 @@ void UCompositeLayerShadowReflection::UpdatePrimitiveVisibilityState(
 	}
 }
 
-void UCompositeLayerShadowReflection::RestorePrimitiveVisibilityState() const
+void UCompositeLayerShadowReflection::SetHiddenPrimitiveState(bool bInHiddenInSceneCapture) const
 {
-	// Restore to class default
-	UPrimitiveComponent* CDO = UPrimitiveComponent::StaticClass()->GetDefaultObject<UPrimitiveComponent>();
+	//* Note: We aim to expose megalights screen traces as an PPV setting, in which case the cvar could be removed. */
+
+	const TCHAR* MegaLightsSceenTracesName = TEXT("r.MegaLights.ScreenTraces");
+	const TCHAR* LumenShortRangeAOHWRTName = TEXT("r.Lumen.ScreenProbeGather.ShortRangeAO.HardwareRayTracing");
+
+	IConsoleVariable* MegaLightsSceenTracesCVar = IConsoleManager::Get().FindConsoleVariable(MegaLightsSceenTracesName);
+	IConsoleVariable* LumenShortRangeAOHWRTCVar = IConsoleManager::Get().FindConsoleVariable(LumenShortRangeAOHWRTName);
+
+	static TSortedMap<FString, int32> CachedConsoleVariables;
+	static TSet<const UCompositeLayerShadowReflection*> ActiveHiddenStateLayers;
+
+	if (bInHiddenInSceneCapture)
+	{
+		constexpr bool bHiddenInSceneCapture = true;
+		constexpr bool bAffectIndirectLightingWhileHidden = true;
+		constexpr bool bCastHiddenShadow = true;
+		UpdatePrimitiveVisibilityState(bHiddenInSceneCapture, bAffectIndirectLightingWhileHidden, bCastHiddenShadow);
+
+		if (ActiveHiddenStateLayers.IsEmpty())
+		{
+			if (MegaLightsSceenTracesCVar != nullptr)
+			{
+				CachedConsoleVariables.FindOrAdd(MegaLightsSceenTracesName) = MegaLightsSceenTracesCVar->GetInt();
+				MegaLightsSceenTracesCVar->Set(0);
+
+				UE_LOG(LogComposite, Display, TEXT("Applying '%s = %d' override for shadow/reflection layer."), MegaLightsSceenTracesName, MegaLightsSceenTracesCVar->GetInt());
+			}
+
+			if (LumenShortRangeAOHWRTCVar != nullptr)
+			{
+				CachedConsoleVariables.FindOrAdd(LumenShortRangeAOHWRTName) = LumenShortRangeAOHWRTCVar->GetInt();
+				LumenShortRangeAOHWRTCVar->Set(1);
+
+				UE_LOG(LogComposite, Display, TEXT("Applying '%s = %d' override for shadow/reflection layer."), LumenShortRangeAOHWRTName, LumenShortRangeAOHWRTCVar->GetInt());
+			}
+		}
+
+		ActiveHiddenStateLayers.Add(this);
+	}
+	else
+	{
+		// Restore to class default
+		UPrimitiveComponent* CDO = UPrimitiveComponent::StaticClass()->GetDefaultObject<UPrimitiveComponent>();
+
+		UpdatePrimitiveVisibilityState(CDO->bHiddenInSceneCapture, CDO->bAffectIndirectLightingWhileHidden, CDO->bCastHiddenShadow);
 
-	UpdatePrimitiveVisibilityState(CDO->bHiddenInSceneCapture, CDO->bAffectIndirectLightingWhileHidden, CDO->bCastHiddenShadow);
+		const int32 NumRemoved = ActiveHiddenStateLayers.Remove(this);
+
+		if (NumRemoved > 0 && ActiveHiddenStateLayers.IsEmpty())
+		{
+			int32 PreviousValue = -1;
+			if (MegaLightsSceenTracesCVar && CachedConsoleVariables.RemoveAndCopyValue(MegaLightsSceenTracesName, PreviousValue))
+			{
+				MegaLightsSceenTracesCVar->Set(PreviousValue);
+
+				UE_LOG(LogComposite, Display, TEXT("Removing shadow/reflection layer override, revert to '%s = %d'."), MegaLightsSceenTracesName, MegaLightsSceenTracesCVar->GetInt());
+			}
+
+			if (LumenShortRangeAOHWRTCVar && CachedConsoleVariables.RemoveAndCopyValue(LumenShortRangeAOHWRTName, PreviousValue))
+			{
+				LumenShortRangeAOHWRTCVar->Set(PreviousValue);
+
+				UE_LOG(LogComposite, Display, TEXT("Removing shadow/reflection layer override, revert to '%s = %d'."), LumenShortRangeAOHWRTName, LumenShortRangeAOHWRTCVar->GetInt());
+			}
+		}
+	}
 }
 
 TStaticArray<TWeakObjectPtr<UCompositeShadowReflectionCatcherComponent>, 2> UCompositeLayerShadowReflection::FindOrCreateSceneCapturePair(ACompositeActor& InOuter) const
diff --git a/Engine/Plugins/Compositing/Composite/Source/Composite/Public/Layers/CompositeLayerShadowReflection.h b/Engine/Plugins/Compositing/Composite/Source/Composite/Public/Layers/CompositeLayerShadowReflection.h
--- a/Engine/Plugins/Compositing/Composite/Source/Composite/Public/Layers/CompositeLayerShadowReflection.h
+++ b/Engine/Plugins/Compositing/Composite/Source/Composite/Public/Layers/CompositeLayerShadowReflection.h
@@ -18,7 +18,7 @@ class USceneCaptureComponent2D;
 
 /** Automatic primitive configuration mode for registered primitives. */
 UENUM(BlueprintType)
-enum class ECompositeHiddenInSceneCaptureConfiguration : uint8
+enum class UE_DEPRECATED(5.8, "ECompositeHiddenInSceneCaptureConfiguration is now deprecated.") ECompositeHiddenInSceneCaptureConfiguration : uint8
 {
 	None UMETA(ToolTip = "No primitive properties are updated"),
 	Visible UMETA(ToolTip = "The following properties are set to false: bHiddenInSceneCapture, bAffectIndirectLightingWhileHidden & bCastHiddenShadow."),
@@ -44,6 +44,7 @@ public:
 	UE_API virtual void OnRemoved(ACompositeActor* LastOwner) override;
 
 	//~ Begin UObject Interface
+	UE_API virtual void PostLoad() override;
 	UE_API virtual void BeginDestroy() override;
 
 #if WITH_EDITOR
@@ -68,22 +69,23 @@ public:
 	UFUNCTION(BlueprintSetter)
 	UE_API void SetActors(TArray<TSoftObjectPtr<AActor>> InActors);
 
-	/**
-	 * Depending on the configuration mode, we automatically enable or disable the following (primitive component) properties:
-	 * 
-	 * * bHiddenInSceneCapture
-	 * * bAffectIndirectLightingWhileHidden
-	 * * bCastHiddenShadow
-	*/
-	UPROPERTY(EditAnywhere, BlueprintReadWrite, Category = "", meta = (DisplayAfter = "Actors", DisplayName = "Auto-Configure Actors"))
-	ECompositeHiddenInSceneCaptureConfiguration AutoConfigureActors;
+	/** Set the enabled state. */
+	UE_API virtual void SetEnabled(bool bInEnabled) override;
+
+#if WITH_EDITORONLY_DATA
+PRAGMA_DISABLE_DEPRECATION_WARNINGS
+	UE_DEPRECATED(5.8, "AutoConfigureActors is now deprecated.")
+	UPROPERTY(meta = (DeprecatedProperty, DeprecationMessage = "AutoConfigureActors is now deprecated."))
+	ECompositeHiddenInSceneCaptureConfiguration AutoConfigureActors_DEPRECATED;
+PRAGMA_ENABLE_DEPRECATION_WARNINGS
+#endif
 
 private:
 	/** Update scene capture visibility on registered mesh actors primitives. */
 	void UpdatePrimitiveVisibilityState(bool bInHiddenInSceneCapture, bool bInAffectIndirectLightingWhileHidden, bool bInCastHiddenShadow) const;
 
-	/** Restore scene capture visibility on registered mesh actors primitives. */
-	void RestorePrimitiveVisibilityState() const;
+	/** Set primitive visibility state and console variables required for hidden in scene capture. */
+	void SetHiddenPrimitiveState(bool bInHiddenInSceneCapture) const;
 
 	/** Find or create scene captures with CG and without CG. */
 	TStaticArray<TWeakObjectPtr<UCompositeShadowReflectionCatcherComponent>, 2> FindOrCreateSceneCapturePair(ACompositeActor& InOuter) const;
