ConvertFromDeviceZ - decrease fudge factor, as it was causing too large discrepancies in the distance between real SceneDepth and ConvertFromDeviceZ(), breaking temporal reprojection and similar operations. With the new factor at 300km from the origin the error is basically 0 (previously it was 62km). At the same it is still small enough where it's safe to do random math operations (in e.g. material shaders) for ConvertFromDeviceZ(0.0) without handling this case explicitly.

#Synced-CL 48607426

diff --git a/Engine/Plugins/Compositing/Composite/Source/Composite/Private/Layers/CompositeLayerSingleLightShadow.cpp b/Engine/Plugins/Compositing/Composite/Source/Composite/Private/Layers/CompositeLayerSingleLightShadow.cpp
--- a/Engine/Plugins/Compositing/Composite/Source/Composite/Private/Layers/CompositeLayerSingleLightShadow.cpp
+++ b/Engine/Plugins/Compositing/Composite/Source/Composite/Private/Layers/CompositeLayerSingleLightShadow.cpp
@@ -45,38 +45,6 @@ namespace UE
 {
 	namespace CompositeCore
 	{
-		/**
-		* Utility function to create the inverse depth projection transform to be used by the shader system.
-		* Copied from SceneView.cpp.
-		*/
-		FVector4f CreateInvDeviceZToWorldZTransform(const FMatrix& ProjMatrix)
-		{
-			float DepthMul = (float)ProjMatrix.M[2][2];
-			float DepthAdd = (float)ProjMatrix.M[3][2];
-
-			if (DepthAdd == 0.f)
-			{
-				// Avoid dividing by 0 in this case
-				DepthAdd = 0.00000001f;
-			}
-
-			bool bIsPerspectiveProjection = ProjMatrix.M[3][3] < 1.0f;
-			if (bIsPerspectiveProjection)
-			{
-				float SubtractValue = DepthMul / DepthAdd;
-
-				// Subtract a tiny number to avoid divide by 0 errors in the shader when a very far distance is decided from the depth buffer.
-				// This fixes fog not being applied to the black background in the editor.
-				SubtractValue -= 0.00000001f;
-
-				return FVector4f( 0.0f, 0.0f, 1.0f / DepthAdd, SubtractValue );
-			}
-			else
-			{
-				return FVector4f((float)(1.0f / ProjMatrix.M[2][2]), (float)(-ProjMatrix.M[3][2] / ProjMatrix.M[2][2] + 1.0f), 0.0f, 1.0f);
-			}
-		}
-
 		class FCompositeSingleLightShadowProxy : public FCompositeCorePassProxy
 		{
 		public:
@@ -398,7 +366,7 @@ void UCompositeLayerSingleLightShadow::GetShadowMatrices(UCompositeSceneCapture2
 	
 	//TODO: Max subject depth logic will need to be updated once support perspective projection.
 	float MaxSubjectDepth = 1.0f;
-	float InvMaxSubjectDepth = 1.0f / MaxSubjectDepth;
+	float InvMaxSubjectDepth = 1.0f / MaxSubjectDepth; 
 	const float ShadowResolutionFractionX = 0.5f;
 	const float ShadowResolutionFractionY = 0.5f;
 
@@ -412,7 +380,7 @@ void UCompositeLayerSingleLightShadow::GetShadowMatrices(UCompositeSceneCapture2
 		);
 
 	OutShadowMatrix = FMatrix44f(WorldToShadowMatrix);
-	OutShadowInvDeviceZToWorldZ = UE::CompositeCore::CreateInvDeviceZToWorldZTransform(ShadowProjectionMatrix);
+	OutShadowInvDeviceZToWorldZ = CreateInvDeviceZToWorldZTransform(ShadowProjectionMatrix);
 }
 
 
diff --git a/Engine/Shaders/Private/Common.ush b/Engine/Shaders/Private/Common.ush
--- a/Engine/Shaders/Private/Common.ush
+++ b/Engine/Shaders/Private/Common.ush
@@ -1358,6 +1358,7 @@ float IsOrthoProjectionFloat()
 // inverse operation of ConvertFromDeviceZ()
 // @param SceneDepth (linear in world units, W)
 // @return DeviceZ (Z/W)
+// WARNING: InvDeviceZToWorldZTransform has a tiny fudge factor added to the projection matrix in order to prevent divide by zero and will cause mismatches with real DeviceZ very far away from camera
 float ConvertToDeviceZ(float SceneDepth)
 {
 	FLATTEN
@@ -1376,6 +1377,7 @@ float ConvertToDeviceZ(float SceneDepth)
 // also see ConvertToDeviceZ()
 // @param DeviceZ value that is stored in the depth buffer (Z/W)
 // @return SceneDepth (linear in world units, W)
+// WARNING: InvDeviceZToWorldZTransform has a tiny fudge factor added to the projection matrix in order to prevent divide by zero and will cause mismatches with real SceneDepth very far away from camera
 float ConvertFromDeviceZ(float DeviceZ)
 {
 	// Supports ortho and perspective, see CreateInvDeviceZToWorldZTransform()
diff --git a/Engine/Source/Runtime/Engine/Private/SceneView.cpp b/Engine/Source/Runtime/Engine/Private/SceneView.cpp
--- a/Engine/Source/Runtime/Engine/Private/SceneView.cpp
+++ b/Engine/Source/Runtime/Engine/Private/SceneView.cpp
@@ -537,7 +537,7 @@ FVector4f CreateInvDeviceZToWorldZTransform(const FMatrix& ProjMatrix)
 
 		// Subtract a tiny number to avoid divide by 0 errors in the shader when a very far distance is decided from the depth buffer.
 		// This fixes fog not being applied to the black background in the editor.
-		SubtractValue -= 0.00000001f;
+		SubtractValue -= 1e-16;
 
 		return FVector4f(
 			0.0f,
diff --git a/Engine/Source/Runtime/Engine/Public/SceneView.h b/Engine/Source/Runtime/Engine/Public/SceneView.h
--- a/Engine/Source/Runtime/Engine/Public/SceneView.h
+++ b/Engine/Source/Runtime/Engine/Public/SceneView.h
@@ -2764,3 +2764,5 @@ ENGINE_API void SetupViewRectUniformBufferParameters(FViewUniformShaderParameter
 	const FViewMatrices& InViewMatrices,
 	const FViewMatrices& InPrevViewMatrice,
 	const FSetupViewUniformParametersInputs& Inputs);
+
+ENGINE_API FVector4f CreateInvDeviceZToWorldZTransform(const FMatrix& ProjMatrix);
\ No newline at end of file
