Composure: Automatically update key color when changing screen type.

#Synced-CL 48823106

diff --git a/Engine/Plugins/Compositing/Composite/Source/Composite/Private/Passes/CompositePassColorKeyer.cpp b/Engine/Plugins/Compositing/Composite/Source/Composite/Private/Passes/CompositePassColorKeyer.cpp
--- a/Engine/Plugins/Compositing/Composite/Source/Composite/Private/Passes/CompositePassColorKeyer.cpp
+++ b/Engine/Plugins/Compositing/Composite/Source/Composite/Private/Passes/CompositePassColorKeyer.cpp
@@ -295,6 +295,26 @@ void UCompositePassColorKeyer::PostEditChangeProperty(FPropertyChangedEvent& Pro
 		FSlateThrottleManager::Get().DisableThrottle(false);
 		bEditorThrottleDisabled = false;
 	}
+
+	if (PropertyChangedEvent.GetMemberPropertyName() == GET_MEMBER_NAME_CHECKED(ThisClass, ScreenType))
+	{
+		switch (ScreenType)
+		{
+		case ECompositeColorKeyerScreenType::Red:
+			KeyColor = FLinearColor::Red;
+			break;
+		case ECompositeColorKeyerScreenType::Green:
+			KeyColor = FLinearColor::Green;
+			break;
+		case ECompositeColorKeyerScreenType::Blue:
+			KeyColor = FLinearColor::Blue;
+			break;
+		
+		default:
+			checkNoEntry();
+			break;
+		}
+	}
 }
 #endif
 
