Composite: Automatically calculate shadow map bounds (and transform) to facilitate single light shadow setup.

#Synced-CL 48836840

diff --git a/Engine/Plugins/Compositing/Composite/Source/Composite/Private/Layers/CompositeLayerSingleLightShadow.cpp b/Engine/Plugins/Compositing/Composite/Source/Composite/Private/Layers/CompositeLayerSingleLightShadow.cpp
--- a/Engine/Plugins/Compositing/Composite/Source/Composite/Private/Layers/CompositeLayerSingleLightShadow.cpp
+++ b/Engine/Plugins/Compositing/Composite/Source/Composite/Private/Layers/CompositeLayerSingleLightShadow.cpp
@@ -128,6 +128,10 @@ namespace UE
 
 UCompositeLayerSingleLightShadow::UCompositeLayerSingleLightShadow(const FObjectInitializer& ObjectInitializer)
 	: Super(ObjectInitializer)
+	, bAutomaticBounds{ true }
+	, OrthographicWidth{ DEFAULT_ORTHOWIDTH }
+	, ShadowMapResolution{ 2048 }
+	, ShadowStrength{ 0.5f }
 {
 	Operation = ECompositeCoreMergeOp::Multiply;
 
@@ -227,9 +231,49 @@ bool UCompositeLayerSingleLightShadow::GetProxy(FCompositeTraversalContext& InCo
 
 	// Update properties that may be dynamic
 	{
-		// Keep exposed properties & transform in sync
-		CachedShadowMapCapture->OrthoWidth = OrthographicWidth;
-		CachedShadowMapCapture->SetWorldTransform(Light->GetTransform());
+		if (bAutomaticBounds)
+		{
+			FBox ShadowCastingBoundingBox = FBox(EForceInit::ForceInitToZero);
+
+			for (const TSoftObjectPtr<AActor>& ShadowCastingActor : ShadowCastingActors)
+			{
+				const AActor* Actor = ShadowCastingActor.Get();
+				if (IsValid(Actor))
+				{
+					ShadowCastingBoundingBox += Actor->GetComponentsBoundingBox();
+				}
+			}
+
+			// Assume shadow is projected onto the ground xy-plane.
+			const FQuat LightRotation = Light->GetTransform().GetRotation();
+			const FVector LightDirection = LightRotation.RotateVector(FVector(1, 0, 0)).GetSafeNormal();
+			const FVector Up = FVector::UpVector;
+			const FVector Location = ShadowCastingBoundingBox.GetCenter() - 2.0 * LightDirection * ShadowCastingBoundingBox.GetSize().GetMax();
+
+			const float CosTheta = FMath::Abs(FVector::DotProduct(LightDirection, FVector::UpVector));
+
+			CachedShadowMapCapture->SetWorldTransform(FTransform(LightRotation, Location));
+			
+			if (CosTheta < UE_KINDA_SMALL_NUMBER)
+			{
+				CachedShadowMapCapture->OrthoWidth = FMath::Max(DEFAULT_ORTHOWIDTH, ShadowCastingBoundingBox.GetSize().GetMax() * 10.0f);
+			}
+			else
+			{
+				const float Theta = FMath::Acos(FMath::Clamp(CosTheta, -1.0f, 1.0f));
+				const float BoxHeight = ShadowCastingBoundingBox.GetSize().Z;
+				const float ShadowLength = BoxHeight * FMath::Tan(Theta);
+				constexpr float ShadowLengthMargin = 2.0f;
+				constexpr float MinOrthoWidth = 1.f;
+
+				CachedShadowMapCapture->OrthoWidth = FMath::Max(MinOrthoWidth, ShadowLengthMargin * ShadowLength);
+			}
+		}
+		else
+		{
+			CachedShadowMapCapture->OrthoWidth = OrthographicWidth;
+			CachedShadowMapCapture->SetWorldTransform(Light->GetTransform());
+		}
 
 		// Acquire render target
 		FCompositeRenderTargetPool::Get().ConditionalAcquireTarget(CachedShadowMapCapture.Get(), CachedShadowMapCapture->TextureTarget, FIntPoint(ShadowMapResolution, ShadowMapResolution), RTF_R32f);
@@ -255,7 +299,7 @@ bool UCompositeLayerSingleLightShadow::GetProxy(FCompositeTraversalContext& InCo
 		ShadowProxy->ShadowStrengthBias = FVector4f(ShadowStrength, 0.0f, 0.0f, 0.0f);
 		GetShadowMatrices(*CachedShadowMapCapture, ShadowProxy->ShadowMatrix, ShadowProxy->ShadowInvDeviceZToWorldZ);
 	}
-	ShadowProxy->LightWorldDirection = FVector3f(Light->GetTransform().TransformVector(FVector(1, 0, 0)));
+	ShadowProxy->LightWorldDirection = FVector3f(CachedShadowMapCapture->GetComponentTransform().TransformVector(FVector(1, 0, 0)));
 
 	FPassInputDeclArray Inputs;
 	Inputs.SetNum(FixedNumLayerInputs);
diff --git a/Engine/Plugins/Compositing/Composite/Source/Composite/Public/Layers/CompositeLayerSingleLightShadow.h b/Engine/Plugins/Compositing/Composite/Source/Composite/Public/Layers/CompositeLayerSingleLightShadow.h
--- a/Engine/Plugins/Compositing/Composite/Source/Composite/Public/Layers/CompositeLayerSingleLightShadow.h
+++ b/Engine/Plugins/Compositing/Composite/Source/Composite/Public/Layers/CompositeLayerSingleLightShadow.h
@@ -56,8 +56,7 @@ public:
 
 public:
 	/**
-	* The reference light, whose exact transform will by used to render a shadow map.
-	* Please make sure it is aligned to your shadow casting scene by piloting the light.
+	* The reference light, whose direction or exact transform will by used to render a shadow map.
 	* 
 	* Only directional lights are currently supported.
 	*/
@@ -72,22 +71,29 @@ public:
 	UFUNCTION(BlueprintSetter)
 	UE_API void SetShadowCastingActors(TArray<TSoftObjectPtr<AActor>> InShadowCaster);
 
+	/** 
+	* When enable, the shadow map render transform and orthographic width are automatically derived from the selected meshes and light angle.
+	* When disabled, the shadow map transform is copied from the light transform, which should be piloted to align with shadow casting geometry.
+	*/
+	UPROPERTY(EditAnywhere, BlueprintReadWrite, Category = "")
+	bool bAutomaticBounds;
+
 	/**
 	* The desired width (in world units) of the shadow map view (ignored if light is not directional).
 	* Reduce size to only contain the shadow casting geometry in view of the light.
 	*/
-	UPROPERTY(EditAnywhere, BlueprintReadWrite, Category = "")
-	float OrthographicWidth = DEFAULT_ORTHOWIDTH;
+	UPROPERTY(EditAnywhere, BlueprintReadWrite, Category = "", meta = (EditCondition = "!bAutomaticBounds", EditConditionHides))
+	float OrthographicWidth;
 
 	/**
 	* Resolution of the shadow map texture. Lower resolutions will cause more blurry shadows.
 	*/
 	UPROPERTY(EditAnywhere, BlueprintReadWrite, Category = "", meta = (ClampMin = "128", UIMin = "128", UIMax="4096"))
-	int32 ShadowMapResolution = 2048;
+	int32 ShadowMapResolution;
 
 	/** Basic shadow strength multiplier, from 0 to 1. */
 	UPROPERTY(EditAnywhere, BlueprintReadWrite, Category = "", meta = (ClampMin = "0.0", ClampMax = "1.0", UIMin = "0.0", UIMax = "1.0"))
-	float ShadowStrength = 0.5f;
+	float ShadowStrength;
 
 #if WITH_EDITORONLY_DATA
 	UE_DEPRECATED(5.8, "Shadow Bias is now deprecated.")
