Composite: Undo removal of reflected properties in CL 47600697 for safer deprecation path.

#Synced-CL 48836840

diff --git a/Engine/Plugins/Compositing/Composite/Source/Composite/Private/Layers/CompositeLayerSceneCapture.cpp b/Engine/Plugins/Compositing/Composite/Source/Composite/Private/Layers/CompositeLayerSceneCapture.cpp
--- a/Engine/Plugins/Compositing/Composite/Source/Composite/Private/Layers/CompositeLayerSceneCapture.cpp
+++ b/Engine/Plugins/Compositing/Composite/Source/Composite/Private/Layers/CompositeLayerSceneCapture.cpp
@@ -23,7 +23,9 @@ UCompositeLayerSceneCapture::UCompositeLayerSceneCapture(const FObjectInitialize
 	}
 }
 
+PRAGMA_DISABLE_DEPRECATION_WARNINGS
 UCompositeLayerSceneCapture::~UCompositeLayerSceneCapture() = default;
+PRAGMA_ENABLE_DEPRECATION_WARNINGS
 
 void UCompositeLayerSceneCapture::OnRemoved(ACompositeActor* LastOwner)
 {
diff --git a/Engine/Plugins/Compositing/Composite/Source/Composite/Private/Layers/CompositeLayerShadowReflection.cpp b/Engine/Plugins/Compositing/Composite/Source/Composite/Private/Layers/CompositeLayerShadowReflection.cpp
--- a/Engine/Plugins/Compositing/Composite/Source/Composite/Private/Layers/CompositeLayerShadowReflection.cpp
+++ b/Engine/Plugins/Compositing/Composite/Source/Composite/Private/Layers/CompositeLayerShadowReflection.cpp
@@ -21,7 +21,9 @@ UCompositeLayerShadowReflection::UCompositeLayerShadowReflection(const FObjectIn
 	}
 }
 
+PRAGMA_DISABLE_DEPRECATION_WARNINGS
 UCompositeLayerShadowReflection::~UCompositeLayerShadowReflection() = default;
+PRAGMA_ENABLE_DEPRECATION_WARNINGS
 
 void UCompositeLayerShadowReflection::OnRemoved(ACompositeActor* LastOwner)
 {
diff --git a/Engine/Plugins/Compositing/Composite/Source/Composite/Public/Layers/CompositeLayerSceneCapture.h b/Engine/Plugins/Compositing/Composite/Source/Composite/Public/Layers/CompositeLayerSceneCapture.h
--- a/Engine/Plugins/Compositing/Composite/Source/Composite/Public/Layers/CompositeLayerSceneCapture.h
+++ b/Engine/Plugins/Compositing/Composite/Source/Composite/Public/Layers/CompositeLayerSceneCapture.h
@@ -93,6 +93,12 @@ private:
 	UPROPERTY(EditAnywhere, BlueprintGetter = GetActors, BlueprintSetter = SetActors, Category = "", meta = (AllowPrivateAccess = true))
 	TArray<TSoftObjectPtr<AActor>> Actors;
 
+#if WITH_EDITORONLY_DATA
+	UE_DEPRECATED(5.8, "CachedVisibilityInSceneCapture is now deprecated.")
+	UPROPERTY(meta = (DeprecatedProperty, DeprecationMessage = "CachedVisibilityInSceneCapture is now deprecated."))
+	TMap<TWeakObjectPtr<UPrimitiveComponent>, bool> CachedVisibilityInSceneCapture;
+#endif
+
 	/** Visibility setting applied to registered actors. By default, registered meshes will only be visible in scene captures and no longer in the main render. */
 	UPROPERTY(EditAnywhere, BlueprintGetter = IsVisibleInSceneCaptureOnly, BlueprintSetter = SetVisibleInSceneCaptureOnly, Category = "", meta = (AllowPrivateAccess = true))
 	bool bVisibleInSceneCaptureOnly;
diff --git a/Engine/Plugins/Compositing/Composite/Source/Composite/Public/Layers/CompositeLayerShadowReflection.h b/Engine/Plugins/Compositing/Composite/Source/Composite/Public/Layers/CompositeLayerShadowReflection.h
--- a/Engine/Plugins/Compositing/Composite/Source/Composite/Public/Layers/CompositeLayerShadowReflection.h
+++ b/Engine/Plugins/Compositing/Composite/Source/Composite/Public/Layers/CompositeLayerShadowReflection.h
@@ -25,6 +25,25 @@ enum class UE_DEPRECATED(5.8, "ECompositeHiddenInSceneCaptureConfiguration is no
 	Hidden UMETA(ToolTip = "The following properties are set to true: bHiddenInSceneCapture, bAffectIndirectLightingWhileHidden & bCastHiddenShadow."),
 };
 
+/** Primitive visibility settings cached state. */
+USTRUCT()
+struct UE_DEPRECATED(5.8, "FCompositeShadowReflectionPrimitiveState is now deprecated.") FCompositeShadowReflectionPrimitiveState
+{
+public:
+	GENERATED_USTRUCT_BODY()
+	/** Hidden in scene capture. */
+	UPROPERTY()
+	bool bHiddenInSceneCapture = false;
+
+	/** Affect indirect lighting while hidden. */
+	UPROPERTY()
+	bool bAffectIndirectLightingWhileHidden = false;
+
+	/** Cast hidden shadow. */
+	UPROPERTY()
+	bool bCastHiddenShadow = false;
+};
+
 /**
 * Layer for shadow & reflection catching with two additional scene capture renders.
 * Primitives are rendered in one scene capture but not the other. The division of both creates a multiplicative matte that can be applied onto the plate.
@@ -95,6 +114,15 @@ private:
 	UPROPERTY(EditAnywhere, BlueprintGetter = GetActors, BlueprintSetter = SetActors, Category = "")
 	TArray<TSoftObjectPtr<AActor>> Actors;
 
+#if WITH_EDITORONLY_DATA
+PRAGMA_DISABLE_DEPRECATION_WARNINGS
+	UE_DEPRECATED(5.8, "CachedVisibilityInSceneCapture is now deprecated.")
+	UPROPERTY(meta = (DeprecatedProperty, DeprecationMessage = "CachedVisibilityInSceneCapture is now deprecated."))
+	TMap<TWeakObjectPtr<UPrimitiveComponent>, FCompositeShadowReflectionPrimitiveState> CachedVisibilityInSceneCapture;
+PRAGMA_ENABLE_DEPRECATION_WARNINGS
+#endif
+
+
 	/** Cached pointer to the scene captures with CG and without CG. */
 	mutable TStaticArray<TWeakObjectPtr<UCompositeShadowReflectionCatcherComponent>, 2> CachedSceneCaptures;
 
