Composure: Improved depth WPO support with lit material & procedural normals.

#virtualized
#Synced-CL 48836840

diff --git a/Engine/Plugins/Compositing/Composite/Source/Composite/Private/Components/CompositeDepthMeshComponent.cpp b/Engine/Plugins/Compositing/Composite/Source/Composite/Private/Components/CompositeDepthMeshComponent.cpp
--- a/Engine/Plugins/Compositing/Composite/Source/Composite/Private/Components/CompositeDepthMeshComponent.cpp
+++ b/Engine/Plugins/Compositing/Composite/Source/Composite/Private/Components/CompositeDepthMeshComponent.cpp
@@ -83,7 +83,7 @@ UCompositeDepthMeshComponent::UCompositeDepthMeshComponent(const FObjectInitiali
 		ConstructorHelpers::FObjectFinder<UMaterialInterface> DefaultDepthMaterial;
 		FConstructorStatics()
 			: DefaultDepthMesh(TEXT("/Composite/Meshes/SM_Grid_960_540.SM_Grid_960_540"))
-			, DefaultDepthMaterial(TEXT("/Composite/Materials/M_CompositeDepthMesh_Unlit_Masked.M_CompositeDepthMesh_Unlit_Masked"))
+			, DefaultDepthMaterial(TEXT("/Composite/Materials/M_CompositeDepthMesh_Lit_Masked.M_CompositeDepthMesh_Lit_Masked"))
 		{
 		}
 	};
