Composure: Add support actor hidden in game property, which will deactivate composite rendering in game mode if enabled.

#Synced-CL 49558946

diff --git a/Engine/Plugins/Compositing/Composite/Source/Composite/Private/CompositeActor.cpp b/Engine/Plugins/Compositing/Composite/Source/Composite/Private/CompositeActor.cpp
--- a/Engine/Plugins/Compositing/Composite/Source/Composite/Private/CompositeActor.cpp
+++ b/Engine/Plugins/Compositing/Composite/Source/Composite/Private/CompositeActor.cpp
@@ -822,16 +822,22 @@ void ACompositeActor::DestroySceneCaptures(const UCompositeLayerBase* InLayer)
 
 bool ACompositeActor::IsRendering() const
 {
-	bool bIsRendering = IsEnabled() && IsActive();
+	// Hidden in game
+	const UWorld* World = GetWorld();
+	if (IsValid(World) && World->IsGameWorld() && IsHidden())
+	{
+		return false;
+	}
 
+	// Hidden in editor
 #if WITH_EDITOR
 	if (IsTemporarilyHiddenInEditor())
 	{
-		bIsRendering = false;
+		return false;
 	}
 #endif
 
-	return bIsRendering;
+	return IsEnabled() && IsActive();
 }
 
 void ACompositeActor::PostJoinConcertSession()
