Composite: Built-in color grading on single light shadow layer.

#Synced-CL 49805347

diff --git a/Engine/Plugins/Compositing/Composite/Source/Composite/Private/Layers/CompositeLayerSingleLightShadow.cpp b/Engine/Plugins/Compositing/Composite/Source/Composite/Private/Layers/CompositeLayerSingleLightShadow.cpp
--- a/Engine/Plugins/Compositing/Composite/Source/Composite/Private/Layers/CompositeLayerSingleLightShadow.cpp
+++ b/Engine/Plugins/Compositing/Composite/Source/Composite/Private/Layers/CompositeLayerSingleLightShadow.cpp
@@ -280,7 +280,8 @@ bool UCompositeLayerSingleLightShadow::GetProxy(FCompositeTraversalContext& InCo
 	}
 	
 	// Setup child shadow pass proxy
-	FCompositeSingleLightShadowProxy* ShadowProxy;
+	FCompositeCorePassProxy* InputProxy;
+
 	{
 		// Display-resolution scene depth without jitter as an external texture.
 		const ResourceId SceneDepthId = InContext.FindOrCreateExternalTexture(CachedSceneDepthCapture->TextureTarget, { .DebugName = TEXT("CompositSceneDepthTex") });
@@ -294,15 +295,36 @@ bool UCompositeLayerSingleLightShadow::GetProxy(FCompositeTraversalContext& InCo
 		PassInputs[0].Set<FPassExternalResourceDesc>({ SceneDepthId });
 		PassInputs[1].Set<FPassExternalResourceDesc>({ ShadowMapId });
 
-		ShadowProxy = InFrameAllocator.Create<FCompositeSingleLightShadowProxy>(PassInputs);
+		FCompositeSingleLightShadowProxy* ShadowProxy = InFrameAllocator.Create<FCompositeSingleLightShadowProxy>(PassInputs);
 		ShadowProxy->ShadowStrengthBias = FVector4f(ShadowStrength, 0.0f, 0.0f, 0.0f);
 		GetShadowMatrices(*CachedShadowMapCapture, ShadowProxy->ShadowMatrix, ShadowProxy->ShadowInvDeviceZToWorldZ);
+		ShadowProxy->LightWorldDirection = FVector3f(CachedShadowMapCapture->GetComponentTransform().TransformVector(FVector(1, 0, 0)));
+
+		// Set input proxy as shadow pass
+		InputProxy = ShadowProxy;
+	}
+
+	if (ColorGrading.bEnableColorGrading)
+	{
+		using namespace UE::Composite::Private;
+
+		FPassInputDecl InputDecl;
+		// Use the shadow result as our color grading pass input
+		InputDecl.Set<const FCompositeCorePassProxy*>(InputProxy);
+
+		// Add our color grading pass
+		FColorGradingCompositePassProxy* GradeProxy = InFrameAllocator.Create<FColorGradingCompositePassProxy>(UE::CompositeCore::FPassInputDeclArray{ InputDecl });
+		GradeProxy->TemperatureSettings = ColorGrading.TemperatureSettings;
+		GradeProxy->ColorGradingSettings = ColorGrading.ColorGradingSettings;
+		GradeProxy->bUnpremult = ColorGrading.bUnpremult;
+
+		// We update the layer input with our color grading pass
+		InputProxy = GradeProxy;
 	}
-	ShadowProxy->LightWorldDirection = FVector3f(CachedShadowMapCapture->GetComponentTransform().TransformVector(FVector(1, 0, 0)));
 
 	FPassInputDeclArray Inputs;
 	Inputs.SetNum(FixedNumLayerInputs);
-	Inputs[0].Set<const FCompositeCorePassProxy*>(ShadowProxy);
+	Inputs[0].Set<const FCompositeCorePassProxy*>(InputProxy);
 	Inputs[1] = GetDefaultSecondInput(InContext);
 
 	OutProxy = InFrameAllocator.Create<FMergePassProxy>(MoveTemp(Inputs), GetMergeOperation(InContext), TEXT("SingleLightShadowCatcher"));
diff --git a/Engine/Plugins/Compositing/Composite/Source/Composite/Private/Passes/CompositePassColorGrading.cpp b/Engine/Plugins/Compositing/Composite/Source/Composite/Private/Passes/CompositePassColorGrading.cpp
--- a/Engine/Plugins/Compositing/Composite/Source/Composite/Private/Passes/CompositePassColorGrading.cpp
+++ b/Engine/Plugins/Compositing/Composite/Source/Composite/Private/Passes/CompositePassColorGrading.cpp
@@ -54,16 +54,14 @@ class FCompositeColorGradingShader : public FGlobalShader
 };
 IMPLEMENT_GLOBAL_SHADER(FCompositeColorGradingShader, "/Plugin/Composite/Private/CompositeColorGrading.usf", "MainPS", SF_Pixel);
 
-class FColorGradingCompositePassProxy : public FCompositeCorePassProxy
+namespace UE::Composite::Private
 {
-public:
-	IMPLEMENT_COMPOSITE_PASS(FColorGradingCompositePassProxy);
-
-	FColorGradingCompositePassProxy(UE::CompositeCore::FPassInputDeclArray InPassDeclaredInputs)
+	FColorGradingCompositePassProxy::FColorGradingCompositePassProxy(CompositeCore::FPassInputDeclArray InPassDeclaredInputs)
 		: FCompositeCorePassProxy(MoveTemp(InPassDeclaredInputs))
-	{ }
+	{
+	}
 
-	UE::CompositeCore::FPassTexture Add(FRDGBuilder& GraphBuilder, const FSceneView& InView, const UE::CompositeCore::FPassInputArray& Inputs, const UE::CompositeCore::FPassContext& PassContext) const override
+	CompositeCore::FPassTexture FColorGradingCompositePassProxy::Add(FRDGBuilder& GraphBuilder, const FSceneView& InView, const CompositeCore::FPassInputArray& Inputs, const CompositeCore::FPassContext& PassContext) const
 	{
 		RDG_EVENT_SCOPE_STAT(GraphBuilder, FCompositeColorGrading, "Composite.ColorGrading");
 
@@ -110,13 +108,9 @@ public:
 			Parameters
 		);
 
-		return UE::CompositeCore::FPassTexture{ MoveTemp(Output), Metadata };
+		return CompositeCore::FPassTexture{ MoveTemp(Output), Metadata };
 	}
-
-	FCompositeTemperatureSettings TemperatureSettings;
-	FColorGradingSettings ColorGradingSettings;
-	bool bUnpremult = true;
-};
+} // End namespace UE::Composite::Private
 
 UCompositePassColorGrading::UCompositePassColorGrading(const FObjectInitializer& ObjectInitializer)
 	: Super(ObjectInitializer)
@@ -130,6 +124,8 @@ UCompositePassColorGrading::~UCompositePassColorGrading() = default;
 
 bool UCompositePassColorGrading::GetProxy(const UE::CompositeCore::FPassInputDecl& InputDecl, FCompositeTraversalContext& InContext, FSceneRenderingBulkObjectAllocator& InFrameAllocator, FCompositeCorePassProxy*& OutProxy) const
 {
+	using namespace UE::Composite::Private;
+
 	FColorGradingCompositePassProxy* Proxy = InFrameAllocator.Create<FColorGradingCompositePassProxy>(UE::CompositeCore::FPassInputDeclArray{ InputDecl });
 	Proxy->TemperatureSettings = TemperatureSettings;
 	Proxy->ColorGradingSettings = ColorGradingSettings;
diff --git a/Engine/Plugins/Compositing/Composite/Source/Composite/Public/Layers/CompositeLayerSingleLightShadow.h b/Engine/Plugins/Compositing/Composite/Source/Composite/Public/Layers/CompositeLayerSingleLightShadow.h
--- a/Engine/Plugins/Compositing/Composite/Source/Composite/Public/Layers/CompositeLayerSingleLightShadow.h
+++ b/Engine/Plugins/Compositing/Composite/Source/Composite/Public/Layers/CompositeLayerSingleLightShadow.h
@@ -2,6 +2,8 @@
 
 #pragma once
 
+#include "Passes/CompositePassColorGrading.h"
+
 #include "EngineDefines.h"
 #include "CompositeLayerBase.h"
 #include "UObject/WeakObjectPtr.h"
@@ -95,6 +97,10 @@ public:
 	UPROPERTY(EditAnywhere, BlueprintReadWrite, Category = "", meta = (ClampMin = "0.0", ClampMax = "1.0", UIMin = "0.0", UIMax = "1.0"))
 	float ShadowStrength;
 
+	/** Color grading pass settings. */
+	UPROPERTY(BlueprintReadWrite, EditAnywhere, Category = "")
+	FCompositeColorGradingPassSettings ColorGrading;
+
 #if WITH_EDITORONLY_DATA
 	UE_DEPRECATED(5.8, "Shadow Bias is now deprecated.")
 	UPROPERTY(meta = (DeprecatedProperty, DeprecationMessage = "Shadow Bias is now deprecated."))
diff --git a/Engine/Plugins/Compositing/Composite/Source/Composite/Public/Passes/CompositePassColorGrading.h b/Engine/Plugins/Compositing/Composite/Source/Composite/Public/Passes/CompositePassColorGrading.h
--- a/Engine/Plugins/Compositing/Composite/Source/Composite/Public/Passes/CompositePassColorGrading.h
+++ b/Engine/Plugins/Compositing/Composite/Source/Composite/Public/Passes/CompositePassColorGrading.h
@@ -5,6 +5,7 @@
 #include "CoreTypes.h"
 #include "CompositePassBase.h"
 #include "Passes/CompositeCorePassMergeProxy.h"
+#include "Engine/Scene.h"
 
 #include "CompositePassColorGrading.generated.h"
 
@@ -31,7 +32,6 @@ struct FCompositeTemperatureSettings
 	/** Controls the color of the scene along the magenta - green axis (orthogonal to the color temperature).  This feature is equivalent to color tint in digital cameras. */
 	UPROPERTY(interp, BlueprintReadWrite, Category = "Color Grading|Temperature", meta = (UIMin = "-1.0", UIMax = "1.0", DisplayName = "Tint"))
 	float WhiteTint = 0.0f;
-
 };
 
 /**
@@ -65,4 +65,56 @@ public:
 	FColorGradingSettings ColorGradingSettings;
 };
 
+USTRUCT(BlueprintType)
+struct FCompositeColorGradingPassSettings
+{
+	GENERATED_USTRUCT_BODY()
+
+	/** Hidden in scene capture. */
+	UPROPERTY(BlueprintReadWrite, EditAnywhere, Category = "")
+	bool bEnableColorGrading = false;
+
+	/** Color is divided by its alpha channel before processing, and multiplied again after. */
+	UPROPERTY(BlueprintReadWrite, EditAnywhere, Category = "")
+	bool bUnpremult = true;
+
+	/** Color temperature settings. */
+	UPROPERTY(BlueprintReadWrite, EditAnywhere, Category = "", meta = (DisplayName = "Temperature"))
+	FCompositeTemperatureSettings TemperatureSettings;
+
+	/** Color grading settings. */
+	UPROPERTY(BlueprintReadWrite, EditAnywhere, Category = "", meta = (ShowOnlyInnerProperties))
+	FColorGradingSettings ColorGradingSettings;
+};
+
+namespace UE
+{
+	namespace Composite
+	{
+		namespace Private
+		{
+			/** Render pass proxy implementation. */
+			class FColorGradingCompositePassProxy : public FCompositeCorePassProxy
+			{
+			public:
+				IMPLEMENT_COMPOSITE_PASS(FColorGradingCompositePassProxy);
+
+				/** Constructor */
+				FColorGradingCompositePassProxy(CompositeCore::FPassInputDeclArray InPassDeclaredInputs);
+
+				CompositeCore::FPassTexture Add(FRDGBuilder& GraphBuilder, const FSceneView& InView, const CompositeCore::FPassInputArray& Inputs, const CompositeCore::FPassContext& PassContext) const override;
+
+				/** Color is divided by its alpha channel before processing, and multiplied again after. */
+				bool bUnpremult = true;
+				
+				/** Color temperature settings. */
+				FCompositeTemperatureSettings TemperatureSettings;
+				
+				/** Color grading settings. */
+				FColorGradingSettings ColorGradingSettings;
+			};
+		}
+	}
+}
+
 #undef UE_API
