Composure: Reorganized Keyer pass properties, and added 'Capture Clean Plate' functionality that grabs a frame from the active media source texture to save and apply as a clean plate texture

#virtualized
#Synced-CL 49870848

diff --git a/Engine/Plugins/Compositing/Composite/Source/Composite/Private/Passes/CompositePassColorKeyer.cpp b/Engine/Plugins/Compositing/Composite/Source/Composite/Private/Passes/CompositePassColorKeyer.cpp
--- a/Engine/Plugins/Compositing/Composite/Source/Composite/Private/Passes/CompositePassColorKeyer.cpp
+++ b/Engine/Plugins/Compositing/Composite/Source/Composite/Private/Passes/CompositePassColorKeyer.cpp
@@ -7,6 +7,7 @@
 #include "PixelShaderUtils.h"
 #include "SystemTextures.h"
 #include "TextureResource.h"
+#include "UObject/FortniteMainBranchObjectVersion.h"
 
 #if WITH_EDITOR
 #include "Application/ThrottleManager.h"
@@ -243,12 +244,32 @@ bool UCompositePassColorKeyer::GetProxy(const UE::CompositeCore::FPassInputDecl&
 	Proxy->DenoiseMethod = DenoiseMethod;
 	Proxy->bInvertAlpha = bInvertAlpha;
 	Proxy->Visualization = Visualization;
-	Proxy->CleanPlateWeakPtr = CleanPlate;
+
+	if (KeyerSource == ECompositeKeyerSource::CleanPlate)
+	{
+		Proxy->CleanPlateWeakPtr = CleanPlate;
+	}
 
 	OutProxy = Proxy;
 	return true;
 }
 
+void UCompositePassColorKeyer::PostLoad()
+{
+	Super::PostLoad();
+
+#if WITH_EDITORONLY_DATA
+	// Set the initial value of KeyerSource to old behavior that assumed clean plate texture should be used if set
+	if (GetLinkerCustomVersion(FFortniteMainBranchObjectVersion::GUID) < FFortniteMainBranchObjectVersion::CompositePassKeyerSource)
+	{
+		if (CleanPlate != nullptr)
+		{
+			KeyerSource = ECompositeKeyerSource::CleanPlate;
+		}
+	}
+#endif
+}
+
 #if WITH_EDITOR
 void UCompositePassColorKeyer::PreEditChange(FProperty* PropertyThatWillChange)
 {
diff --git a/Engine/Plugins/Compositing/Composite/Source/Composite/Public/Passes/CompositePassColorKeyer.h b/Engine/Plugins/Compositing/Composite/Source/Composite/Public/Passes/CompositePassColorKeyer.h
--- a/Engine/Plugins/Compositing/Composite/Source/Composite/Public/Passes/CompositePassColorKeyer.h
+++ b/Engine/Plugins/Compositing/Composite/Source/Composite/Public/Passes/CompositePassColorKeyer.h
@@ -34,6 +34,12 @@ enum class ECompositeDenoiseMethod : uint8
 	Median3x3 UMETA(DisplayName = "Median 3x3"),
 };
 
+UENUM(BlueprintType)
+enum class ECompositeKeyerSource : uint8
+{
+	Color,
+	CleanPlate
+};
 
 /**
  * Color keyer with clean plate & spill removal support.
@@ -52,6 +58,8 @@ public:
 
 	UE_API virtual bool GetProxy(const UE::CompositeCore::FPassInputDecl& InputDecl, FCompositeTraversalContext& InContext, FSceneRenderingBulkObjectAllocator& InFrameAllocator, FCompositeCorePassProxy*& OutProxy) const override;
 
+	UE_API virtual void PostLoad() override;
+	
 #if WITH_EDITOR
 	UE_API virtual void PreEditChange(FProperty* PropertyThatWillChange) override;
 	UE_API virtual void PostEditChangeProperty(FPropertyChangedEvent& PropertyChangedEvent) override;
@@ -62,17 +70,26 @@ public:
 	UPROPERTY(EditAnywhere, BlueprintReadWrite, Category = "")
 	ECompositeColorKeyerScreenType ScreenType;
 
+	UPROPERTY(EditAnywhere, BlueprintReadWrite, Category = "", meta = (PropertyGroup = "Keyer Settings"))
+	ECompositeKeyerSource KeyerSource;
+	
 	/** Static background key color.*/
-	UPROPERTY(EditAnywhere, BlueprintReadWrite, Category = "", meta = (OnlyUpdateOnInteractionEnd, EditCondition = "CleanPlate == nullptr"))
+	UPROPERTY(EditAnywhere, BlueprintReadWrite, Category = "", meta = (PropertyGroup = "Keyer Settings", OnlyUpdateOnInteractionEnd, EditCondition = "KeyerSource == ECompositeKeyerSource::Color", EditConditionHides))
 	FLinearColor KeyColor;
 
 	/**
 	 * Clean plate background for calculating color differences per pixel, instead of the static key color.
 	 * Resolution must match the composite actor render resolution.
 	 */
-	UPROPERTY(EditAnywhere, BlueprintReadWrite, Category = "")
+	UPROPERTY(EditAnywhere, BlueprintReadWrite, Category = "", meta = (PropertyGroup = "Keyer Settings", OnlyUpdateOnInteractionEnd, EditCondition = "KeyerSource == ECompositeKeyerSource::CleanPlate", EditConditionHides))
 	TObjectPtr<UTexture> CleanPlate;
 
+#if WITH_EDITORONLY_DATA
+	/** Setting that controls how long to delay before capturing a clean plate */
+	UPROPERTY(EditAnywhere, Category = "", meta = (PropertyGroup = "Keyer Settings", EditCondition = "KeyerSource == ECompositeKeyerSource::CleanPlate", EditConditionHides, ClampMin=0.0, UIMin=0.0, ClampMax=60.0, UIMax=60.0))
+	float CleanPlateCountdown = 0.0f;
+#endif
+	
 	/** Weight of the foreground red channel contributing to the key matte hardness. */
 	UPROPERTY(EditAnywhere, BlueprintReadWrite, Category = "", meta = (UIMin = 0.0f, UIMax = 1.0f, EditCondition = "ScreenType != ECompositeColorKeyerScreenType::Red"))
 	float RedWeight;
@@ -98,11 +115,11 @@ public:
 	float DevignetteStrength;
 
 	/** When enabled, we undo devignetting before outputting the keyed plate. */
-	UPROPERTY(EditAnywhere, BlueprintReadWrite, Category = "")
+	UPROPERTY(EditAnywhere, BlueprintReadWrite, Category = "", AdvancedDisplay)
 	bool bPreserveVignetteAfterKey;
 
 	/** Denoising method applied before the keyer. */
-	UPROPERTY(EditAnywhere, BlueprintReadWrite, Category = "")
+	UPROPERTY(EditAnywhere, BlueprintReadWrite, Category = "", AdvancedDisplay)
 	ECompositeDenoiseMethod DenoiseMethod;
 
 	/** Vizualize the alpha key or fill. */
@@ -110,13 +127,16 @@ public:
 	ECompositeColorKeyerVisualization Visualization;
 
 	/** Invert the alpha key. */
-	UPROPERTY(EditAnywhere, BlueprintReadWrite, Category = "")
+	UPROPERTY(EditAnywhere, BlueprintReadWrite, Category = "", AdvancedDisplay)
 	bool bInvertAlpha;
 
 #if WITH_EDITOR
 	/** Flag tracking whether we disabled throttling to later restore it. */
 	bool bEditorThrottleDisabled = false;
 #endif
+
+private:
+	friend class FCompositePassColorKeyerCustomization;
 };
 
 #undef UE_API
diff --git a/Engine/Plugins/Compositing/Composite/Source/CompositeEditor/CompositeEditor.Build.cs b/Engine/Plugins/Compositing/Composite/Source/CompositeEditor/CompositeEditor.Build.cs
--- a/Engine/Plugins/Compositing/Composite/Source/CompositeEditor/CompositeEditor.Build.cs
+++ b/Engine/Plugins/Compositing/Composite/Source/CompositeEditor/CompositeEditor.Build.cs
@@ -23,8 +23,10 @@ public class CompositeEditor : ModuleRules
 				"ColorGradingEditor",
 				"CompositeCore",
 				"ConcertSyncClient",
+				"ContentBrowserData",
 				"CoreUObject",
 				"EditorWidgets",
+				"EditorScriptingUtilities",
 				"Engine",
 				"InputCore",
 				"LevelEditor",
diff --git a/Engine/Plugins/Compositing/Composite/Source/CompositeEditor/Private/CompositeEditorModule.cpp b/Engine/Plugins/Compositing/Composite/Source/CompositeEditor/Private/CompositeEditorModule.cpp
--- a/Engine/Plugins/Compositing/Composite/Source/CompositeEditor/Private/CompositeEditorModule.cpp
+++ b/Engine/Plugins/Compositing/Composite/Source/CompositeEditor/Private/CompositeEditorModule.cpp
@@ -28,11 +28,13 @@
 #include "Customizations/CompositeLayerShadowReflectionCustomization.h"
 #include "Customizations/CompositeLayerSingleLightShadowCustomization.h"
 #include "Customizations/CompositeDepthMeshComponentCustomization.h"
+#include "Customizations/CompositePassColorKeyerCustomization.h"
 #include "Layers/CompositeLayerPlate.h"
 #include "Layers/CompositeLayerSceneCapture.h"
 #include "Layers/CompositeLayerShadowReflection.h"
 #include "Layers/CompositeLayerSingleLightShadow.h"
 #include "Passes/CompositePassColorGrading.h"
+#include "Passes/CompositePassColorKeyer.h"
 #include "UI/SCompositeEditorPanel.h"
 
 #define LOCTEXT_NAMESPACE "CompositeEditorModule"
@@ -275,6 +277,11 @@ void FCompositeEditorModule::RegisterCustomizations()
 		UCompositeDepthMeshComponent::StaticClass()->GetFName(),
 		FOnGetDetailCustomizationInstance::CreateStatic(&FCompositeDepthMeshComponentCustomization::MakeInstance)
 	);
+
+	PropertyModule.RegisterCustomClassLayout(
+		UCompositePassColorKeyer::StaticClass()->GetFName(),
+		FOnGetDetailCustomizationInstance::CreateStatic(&FCompositePassColorKeyerCustomization::MakeInstance)
+	);
 }
 
 void FCompositeEditorModule::UnregisterCustomizations()
diff --git a/Engine/Plugins/Compositing/Composite/Source/CompositeEditor/Private/CompositeEditorStyle.cpp b/Engine/Plugins/Compositing/Composite/Source/CompositeEditor/Private/CompositeEditorStyle.cpp
--- a/Engine/Plugins/Compositing/Composite/Source/CompositeEditor/Private/CompositeEditorStyle.cpp
+++ b/Engine/Plugins/Compositing/Composite/Source/CompositeEditor/Private/CompositeEditorStyle.cpp
@@ -53,4 +53,6 @@ FCompositeEditorStyle::FCompositeEditorStyle()
 	Set("CompositeEditor.Passes.Layer", new CORE_IMAGE_BRUSH_SVG("Starship/Common/Layers", Icon16x16));
 
 	Set("CompositeEditor.MediaProfile", new IMAGE_BRUSH_SVG("Editor/Icons/MediaProfile", Icon16x16));
+	Set("CompositeEditor.CaptureCleanPlate", new IMAGE_BRUSH_SVG("Editor/Icons/TextureCreate_16", Icon16x16));
+	Set("CompositeEditor.CancelCapture", new IMAGE_BRUSH_SVG("Editor/Icons/TextureCancel_16", Icon16x16));
 }
diff --git a/Engine/Plugins/Compositing/Composite/Source/CompositeEditor/Private/Customizations/CompositePassBaseCustomization.cpp b/Engine/Plugins/Compositing/Composite/Source/CompositeEditor/Private/Customizations/CompositePassBaseCustomization.cpp
--- a/Engine/Plugins/Compositing/Composite/Source/CompositeEditor/Private/Customizations/CompositePassBaseCustomization.cpp
+++ b/Engine/Plugins/Compositing/Composite/Source/CompositeEditor/Private/Customizations/CompositePassBaseCustomization.cpp
@@ -2,6 +2,7 @@
 
 #include "CompositePassBaseCustomization.h"
 
+#include "DetailCategoryBuilder.h"
 #include "DetailLayoutBuilder.h"
 #include "Passes/CompositePassBase.h"
 
@@ -13,6 +14,11 @@ TSharedRef<IDetailCustomization> FCompositePassBaseCustomization::MakeInstance()
 }
 
 void FCompositePassBaseCustomization::CustomizeDetails(IDetailLayoutBuilder& DetailLayout)
+{
+	CustomizeIsEnabledProperty(DetailLayout);
+}
+
+void FCompositePassBaseCustomization::CustomizeIsEnabledProperty(IDetailLayoutBuilder& DetailLayout, bool bInPlace)
 {
 	TArray<TWeakObjectPtr<UObject>> Objects;
 	DetailLayout.GetObjectsBeingCustomized(Objects);
@@ -31,8 +37,20 @@ void FCompositePassBaseCustomization::CustomizeDetails(IDetailLayoutBuilder& Det
 		ObjectTypeName = LOCTEXT("AllPassesLabel", "All Passes");
 	}
 	
-	TSharedPtr<IPropertyHandle> EnabledHandle = DetailLayout.GetProperty(GET_MEMBER_NAME_CHECKED(UCompositePassBase, bIsEnabled));
-	DetailLayout.EditDefaultProperty(EnabledHandle)->DisplayName(FText::Format(LOCTEXT("PassEnabledLabelFormat", "Enable {0}"), ObjectTypeName));
+	TSharedPtr<IPropertyHandle> EnabledHandle = DetailLayout.GetProperty(GET_MEMBER_NAME_CHECKED(UCompositePassBase, bIsEnabled), UCompositePassBase::StaticClass());
+
+	IDetailPropertyRow* PropertyRow;
+	if (bInPlace)
+	{
+		PropertyRow = DetailLayout.EditDefaultProperty(EnabledHandle);
+	}
+	else
+	{
+		IDetailCategoryBuilder& DefaultCategory = DetailLayout.EditCategoryAllowNone(NAME_None);
+		PropertyRow = &DefaultCategory.AddProperty(EnabledHandle);
+	}
+
+	PropertyRow->DisplayName(FText::Format(LOCTEXT("PassEnabledLabelFormat", "Enable {0}"), ObjectTypeName));
 }
 
 #undef LOCTEXT_NAMESPACE
\ No newline at end of file
diff --git a/Engine/Plugins/Compositing/Composite/Source/CompositeEditor/Private/Customizations/CompositePassBaseCustomization.h b/Engine/Plugins/Compositing/Composite/Source/CompositeEditor/Private/Customizations/CompositePassBaseCustomization.h
--- a/Engine/Plugins/Compositing/Composite/Source/CompositeEditor/Private/Customizations/CompositePassBaseCustomization.h
+++ b/Engine/Plugins/Compositing/Composite/Source/CompositeEditor/Private/Customizations/CompositePassBaseCustomization.h
@@ -13,4 +13,12 @@ public:
 	static TSharedRef<IDetailCustomization> MakeInstance();
 	
 	virtual void CustomizeDetails(IDetailLayoutBuilder& DetailLayout) override;
+
+protected:
+	/**
+	 * Renames the IsEnabled property on the composite pass to reference the property owner's type (e.g. 'Enable Color Grading Pass')
+	 * @param DetailLayout The detail layout builder
+	 * @param bInPlace Whether the IsEnabled property should remain in its default position or be placed with the rest of the customized properties
+	 */
+	void CustomizeIsEnabledProperty(IDetailLayoutBuilder& DetailLayout, bool bInPlace = true);
 };
diff --git a/Engine/Plugins/Compositing/Composite/Source/CompositeEditor/Private/Customizations/CompositePassColorKeyerCustomization.cpp b/Engine/Plugins/Compositing/Composite/Source/CompositeEditor/Private/Customizations/CompositePassColorKeyerCustomization.cpp
new file mode 100644
--- /dev/null
+++ b/Engine/Plugins/Compositing/Composite/Source/CompositeEditor/Private/Customizations/CompositePassColorKeyerCustomization.cpp
@@ -0,0 +1,567 @@
+﻿// Copyright Epic Games, Inc. All Rights Reserved.
+
+#include "CompositePassColorKeyerCustomization.h"
+
+#include "AssetToolsModule.h"
+#include "CompositeEditorModule.h"
+#include "CompositeEditorStyle.h"
+#include "ContentBrowserModule.h"
+#include "DetailCategoryBuilder.h"
+#include "DetailLayoutBuilder.h"
+#include "DetailWidgetRow.h"
+#include "Editor.h"
+#include "EditorAssetLibrary.h"
+#include "FileHelpers.h"
+#include "IContentBrowserSingleton.h"
+#include "IDetailGroup.h"
+#include "MediaTexture.h"
+#include "NumericPropertyParams.h"
+#include "ScopedTransaction.h"
+#include "AssetRegistry/AssetRegistryModule.h"
+#include "Engine/Texture2D.h"
+#include "Kismet/KismetRenderingLibrary.h"
+#include "Layers/CompositeLayerPlate.h"
+#include "Materials/Material.h"
+#include "Materials/MaterialInstanceDynamic.h"
+#include "Passes/CompositePassColorKeyer.h"
+#include "Widgets/Input/SButton.h"
+#include "Widgets/Input/SVectorInputBox.h"
+#include "Widgets/Layout/SWidgetSwitcher.h"
+#include "Widgets/Notifications/SProgressBar.h"
+#include "Widgets/Text/STextBlock.h"
+
+#define LOCTEXT_NAMESPACE "FCompositePassColorKeyerCustomization"
+
+class FContentBrowserModule;
+
+/** Widget that manages a button with an optional delay, after which its Action delegate is invoked */
+class SDelayedActionButton : public SCompoundWidget
+{
+public:
+	SLATE_BEGIN_ARGS(SDelayedActionButton) {}
+		SLATE_NAMED_SLOT(FArguments, ButtonContent)
+		SLATE_NAMED_SLOT(FArguments, CancelContent)
+		SLATE_ATTRIBUTE(float, Delay)
+		SLATE_EVENT(FSimpleDelegate, OnActivateAction)
+	SLATE_END_ARGS()
+
+	void Construct(const FArguments& InArgs)
+	{
+		Delay = InArgs._Delay;
+		OnActivateAction = InArgs._OnActivateAction;
+
+		ChildSlot
+		[
+			SNew(SOverlay)
+		
+			+SOverlay::Slot()
+			.Padding(0.0f, 2.0f)
+			[
+				SNew(SButton)
+				.OnClicked(this, &SDelayedActionButton::OnButtonClicked)
+				[
+					SNew(SWidgetSwitcher)
+					.WidgetIndex(this, &SDelayedActionButton::GetButtonWidgetIndex)
+
+					+SWidgetSwitcher::Slot()
+					[
+						InArgs._ButtonContent.Widget
+					]
+
+					+SWidgetSwitcher::Slot()
+					[
+						InArgs._CancelContent.Widget
+					]
+				]
+			]
+		
+			+SOverlay::Slot()
+			.Padding(0.0f, 2.0f)
+			.VAlign(VAlign_Bottom)
+			[
+				SNew(SVerticalBox)
+				
+				+SVerticalBox::Slot()
+				.MaxHeight(2.0f)
+				[
+					SNew(SProgressBar)
+					.Percent(this, &SDelayedActionButton::GetDelayProgress)
+					.Visibility(this, &SDelayedActionButton::GetProgressBarVisibility)
+				]
+			]
+		];
+	}
+	
+private:
+	FReply OnButtonClicked()
+	{
+		// If there is active delay action timer, clicking on the button should cancel it
+		if (DelayedActionTimer.IsValid())
+		{
+			GEditor->GetTimerManager()->ClearTimer(DelayedActionTimer);
+			return FReply::Handled();
+		}
+
+		// If the amount to delay is non-negligible, start the timer; otherwise, invoke the action immediately
+		const float AmountToDelay = Delay.Get(0.0f);
+		if (AmountToDelay > UE_KINDA_SMALL_NUMBER)
+		{
+			FTimerManagerTimerParameters Args;
+			GEditor->GetTimerManager()->SetTimer(DelayedActionTimer, FTimerDelegate::CreateSP(this, &SDelayedActionButton::OnTimerExpired), AmountToDelay, Args);
+		}
+		else
+		{
+			OnActivateAction.ExecuteIfBound();
+		}
+
+		return FReply::Handled();
+	}
+
+	int GetButtonWidgetIndex() const
+	{
+		return DelayedActionTimer.IsValid() ? 1 : 0;
+	}
+
+	TOptional<float> GetDelayProgress() const
+	{
+		if (DelayedActionTimer.IsValid())
+		{
+			// Compute the fraction of time elapsed to total delay time to display with the progress bar
+			const float TimeElapsed = GEditor->GetTimerManager()->GetTimerElapsed(DelayedActionTimer);
+			const float TimeRemaining = GEditor->GetTimerManager()->GetTimerRemaining(DelayedActionTimer);
+			return TimeElapsed / (TimeElapsed + TimeRemaining);
+		}
+		
+		return TOptional<float>();
+	}
+
+	EVisibility GetProgressBarVisibility() const
+	{
+		return DelayedActionTimer.IsValid() ? EVisibility::Visible : EVisibility::Hidden;
+	}
+
+	void OnTimerExpired()
+	{
+		// Once the timer has expired, invoke the action delegate and clear the timer
+		GEditor->GetTimerManager()->ClearTimer(DelayedActionTimer);
+		OnActivateAction.ExecuteIfBound();
+	}
+	
+private:
+	/** Attribute to retrieve the amount of time to delay before invoking the action */
+	TAttribute<float> Delay;
+
+	/** Invoked once the delay duration has elapsed after clicking the button */
+	FSimpleDelegate OnActivateAction;
+
+	/** Handle for the delay timer */
+	FTimerHandle DelayedActionTimer;
+};
+
+TSharedRef<IDetailCustomization> FCompositePassColorKeyerCustomization::MakeInstance()
+{
+	return MakeShared<FCompositePassColorKeyerCustomization>();
+}
+
+void FCompositePassColorKeyerCustomization::CustomizeDetails(const TSharedPtr<IDetailLayoutBuilder>& DetailBuilder)
+{
+	UMaterial* CaptureCleanPlateMat = LoadObject<UMaterial>(nullptr, TEXT("/Composite/Materials/M_Composite_CleanPlateCapture.M_Composite_CleanPlateCapture"));
+	CaptureCleanPlateMID = TStrongObjectPtr<UMaterialInstanceDynamic>(UMaterialInstanceDynamic::Create(CaptureCleanPlateMat, nullptr));
+	
+	CachedDetailBuilder = DetailBuilder;
+	IDetailCustomization::CustomizeDetails(DetailBuilder);
+}
+
+void FCompositePassColorKeyerCustomization::CustomizeDetails(IDetailLayoutBuilder& DetailBuilder)
+{
+	// All UCompositePassColorKeyer properties reside in the default category (Category = ""). If this changes, will need to update this category 
+	IDetailCategoryBuilder& DefaultCategory = DetailBuilder.EditCategoryAllowNone(NAME_None);
+
+	// 'Customize' properties that should be above the 'Keyer Settings' custom group
+	DefaultCategory.AddProperty(DetailBuilder.GetProperty(GET_MEMBER_NAME_CHECKED(UCompositePassColorKeyer, DisplayName), UCompositePassBase::StaticClass()));
+	CustomizeIsEnabledProperty(DetailBuilder, false);
+	DefaultCategory.AddProperty(DetailBuilder.GetProperty(GET_MEMBER_NAME_CHECKED(UCompositePassColorKeyer, ScreenType)));
+
+	KeyerSourceHandle = DetailBuilder.GetProperty(GET_MEMBER_NAME_CHECKED(UCompositePassColorKeyer, KeyerSource));
+
+	// Add keyer color and clean plate properties to custom 'Keyer Settings' group
+	IDetailGroup& KeyerSettingsGroup = DefaultCategory.AddGroup(TEXT("KeyerSettings"), LOCTEXT("KeyerSettingsGroupName", "Keyer Settings"), false, true);
+	KeyerSettingsGroup.AddPropertyRow(KeyerSourceHandle.ToSharedRef());
+	KeyerSettingsGroup.AddPropertyRow(DetailBuilder.GetProperty(GET_MEMBER_NAME_CHECKED(UCompositePassColorKeyer, KeyColor)));
+	KeyerSettingsGroup.AddPropertyRow(DetailBuilder.GetProperty(GET_MEMBER_NAME_CHECKED(UCompositePassColorKeyer, CleanPlate)));
+
+	CountdownHandle = DetailBuilder.GetProperty(GET_MEMBER_NAME_CHECKED(UCompositePassColorKeyer, CleanPlateCountdown));
+	KeyerSettingsGroup.AddPropertyRow(CountdownHandle.ToSharedRef());
+
+	// Add custom 'Capture Clean Plate' button to 'Keyer Settings' group
+	KeyerSettingsGroup.AddWidgetRow()
+	.WholeRowContent()
+	[
+		CreateCaptureButtonWidget()
+	]
+	.Visibility(TAttribute<EVisibility>::CreateSP(this, &FCompositePassColorKeyerCustomization::GetCaptureButtonVisibility));
+
+	RedWeightHandle = DetailBuilder.GetProperty(GET_MEMBER_NAME_CHECKED(UCompositePassColorKeyer, RedWeight));
+	GreenWeightHandle = DetailBuilder.GetProperty(GET_MEMBER_NAME_CHECKED(UCompositePassColorKeyer, GreenWeight));
+	BlueWeightHandle = DetailBuilder.GetProperty(GET_MEMBER_NAME_CHECKED(UCompositePassColorKeyer, BlueWeight));
+
+	// Create custom group for RGB weight properties, which will include a single row of numeric sliders for each component in the group's header row
+	FText ColorWeightGroupName = LOCTEXT("ColorWeightsGroupName", "Weights");
+	IDetailGroup& ColorWeightGroup = DefaultCategory.AddGroup(TEXT("ColorWeights"), ColorWeightGroupName);
+	ColorWeightGroup.HeaderRow()
+	.NameContent()
+	[
+		SNew(STextBlock)
+		.Font(IDetailLayoutBuilder::GetDetailFont())
+		.Text(ColorWeightGroupName)
+	]
+	.ValueContent()
+	[
+		CreateRGBNumericEntryWidget()
+	];
+	
+	ColorWeightGroup.AddPropertyRow(RedWeightHandle.ToSharedRef());
+	ColorWeightGroup.AddPropertyRow(GreenWeightHandle.ToSharedRef());
+	ColorWeightGroup.AddPropertyRow(BlueWeightHandle.ToSharedRef());
+}
+
+void FCompositePassColorKeyerCustomization::CaptureMediaTexture()
+{
+	TSharedPtr<IDetailLayoutBuilder> DetailBuilder = CachedDetailBuilder.Pin();
+	if (!DetailBuilder.IsValid())
+	{
+		return;
+	}
+	
+	TArray<TWeakObjectPtr<UCompositePassColorKeyer>> Keyers = DetailBuilder->GetObjectsOfTypeBeingCustomized<UCompositePassColorKeyer>();
+	if (Keyers.Num() != 1)
+	{
+		return;
+	}
+
+	TStrongObjectPtr<UCompositePassColorKeyer> PinnedKeyer = Keyers[0].Pin();
+	if (!PinnedKeyer.IsValid())
+	{
+		return;
+	}
+	
+	UCompositeLayerPlate* ParentPlate = PinnedKeyer->GetTypedOuter<UCompositeLayerPlate>();
+	if (!ParentPlate)
+	{
+		return;
+	}
+
+	UMediaTexture* MediaTexture = Cast<UMediaTexture>(ParentPlate->Texture);
+	if (!MediaTexture)
+	{
+		return;
+	}
+
+	// Set the media texture as a parameter on the capture material
+	CaptureCleanPlateMID->SetTextureParameterValue(TEXT("MediaTexture"), MediaTexture);
+
+	// Draw material to render target using an linear HDR format
+	UTextureRenderTarget2D* RenderTarget = NewObject<UTextureRenderTarget2D>();
+	RenderTarget->InitCustomFormat(MediaTexture->GetWidth(), MediaTexture->GetHeight(), PF_FloatRGBA, true);
+	
+	UKismetRenderingLibrary::ClearRenderTarget2D(PinnedKeyer.Get(), RenderTarget, FLinearColor::Black);
+	UKismetRenderingLibrary::DrawMaterialToRenderTarget(PinnedKeyer.Get(), RenderTarget, CaptureCleanPlateMID.Get());
+	RenderTarget->UpdateResourceImmediate(false);
+
+	// Query user for location to save the texture asset
+	IContentBrowserSingleton& ContentBrowser = FModuleManager::LoadModuleChecked<FContentBrowserModule>("ContentBrowser").Get();
+	const FContentBrowserItemPath CurrentPath = ContentBrowser.GetCurrentPath();
+	
+	FSaveAssetDialogConfig SaveAssetDialogConfig;
+	SaveAssetDialogConfig.DialogTitleOverride = LOCTEXT("SaveCleanPlateDialogTitle", "Save as Clean Plate Texture");
+	SaveAssetDialogConfig.DefaultPath = CurrentPath.HasInternalPath() ? CurrentPath.GetInternalPathString() : TEXT("/Game");
+	SaveAssetDialogConfig.DefaultAssetName = TEXT("NewCleanPlate");
+	SaveAssetDialogConfig.ExistingAssetPolicy = ESaveAssetDialogExistingAssetPolicy::AllowButWarn;
+	SaveAssetDialogConfig.AssetClassNames.Add(UTexture2D::StaticClass()->GetClassPathName());
+
+	FString TextureAssetPath = ContentBrowser.CreateModalSaveAssetDialog(SaveAssetDialogConfig);
+	if (!TextureAssetPath.IsEmpty())
+	{
+		const FString RequestedName = FPaths::GetBaseFilename(TextureAssetPath, true);
+		const FString RequestedPath = FPaths::GetPath(TextureAssetPath);
+		const FString RequestedFullPath = FPaths::Combine(RequestedPath, RequestedName);
+
+		// If there already exists a texture asset at the user's chosen location, we will overwrite it instead of creating a new one
+		UTexture2D* CleanPlate = nullptr;
+		if (UEditorAssetLibrary::DoesAssetExist(RequestedFullPath))
+		{
+			CleanPlate = Cast<UTexture2D>(UEditorAssetLibrary::LoadAsset(RequestedFullPath));
+			if (CleanPlate)
+			{
+				if (!RenderTarget->UpdateTexture(CleanPlate))
+				{
+					CleanPlate = nullptr;
+				}
+			}
+			else
+			{
+				// Asset exists, but for some reason isn't a 2D texture. Abort
+				UE_LOG(LogCompositeEditor, Display, TEXT("Could not overwrite %s with clean plate as it is not a valid texture asset"), *RequestedFullPath);
+				return;
+			}
+		}
+		else
+		{
+			FString FinalPackageName, FinalAssetName;
+			const FAssetToolsModule& AssetToolsModule = FModuleManager::LoadModuleChecked<FAssetToolsModule>("AssetTools");
+			AssetToolsModule.Get().CreateUniqueAssetName(FPaths::Combine(RequestedPath, RequestedName), TEXT(""), FinalPackageName, FinalAssetName);
+
+			UPackage* NewPackage = CreatePackage(*FinalPackageName);
+			CleanPlate = RenderTarget->ConstructTexture2D(NewPackage, RequestedName, RF_Public | RF_Standalone | RF_Transactional);
+			
+			FAssetRegistryModule::AssetCreated(CleanPlate);
+		}
+
+		if (!CleanPlate)
+		{
+			UE_LOG(LogCompositeEditor, Display, TEXT("Failed to create texture asset for clean plate at %s"), *RequestedFullPath);
+			return;
+		}
+		
+		CleanPlate->UpdateResource();
+		CleanPlate->MarkPackageDirty();
+
+		// Automatically save the newly created or overwritten texture asset, and update the UCompositePassColorKeyer's CleanPlate property with the new asset
+		if (UEditorLoadingAndSavingUtils::SavePackages({ CleanPlate->GetPackage() }, true))
+		{
+			const FScopedTransaction Transaction(LOCTEXT("SetCleanPlateTransaction", "Set Clean Plate"));
+			PinnedKeyer->Modify();
+
+			FProperty* CleanPlateProperty = FindFProperty<FProperty>(UCompositePassColorKeyer::StaticClass(), GET_MEMBER_NAME_CHECKED(UCompositePassColorKeyer, CleanPlate));
+			PinnedKeyer->PreEditChange(CleanPlateProperty);
+			PinnedKeyer->CleanPlate = CleanPlate;
+
+			FPropertyChangedEvent ChangedEvent(CleanPlateProperty, EPropertyChangeType::ValueSet, { PinnedKeyer.Get() });
+			PinnedKeyer->PostEditChangeProperty(ChangedEvent);
+		}
+	}
+}
+
+TSharedRef<SWidget> FCompositePassColorKeyerCustomization::CreateRGBNumericEntryWidget()
+{
+	auto CreateNumericEntryBox = [this](const TSharedPtr<IPropertyHandle>& InPropertyHandle, FLinearColor InLabelColor)
+	{
+		return SNew(SNumericEntryBox<float>)
+			.Value(this, &FCompositePassColorKeyerCustomization::GetRGBWeightValue, InPropertyHandle)
+			.Font(IDetailLayoutBuilder::GetDetailFont())
+			.AllowSpin(true)
+			.IsEnabled(this, &FCompositePassColorKeyerCustomization::IsRGBWeightEditable, InPropertyHandle)
+			.OnValueChanged(this, &FCompositePassColorKeyerCustomization::ChangeRGBWeightValue, InPropertyHandle)
+			.OnValueCommitted(this, &FCompositePassColorKeyerCustomization::CommitRGBWeightValue, InPropertyHandle)
+			.MinValue(this, &FCompositePassColorKeyerCustomization::GetRGBWeightMin, InPropertyHandle)
+			.MinSliderValue(this, &FCompositePassColorKeyerCustomization::GetRGBWeightMin, InPropertyHandle)
+			.MaxValue(this, &FCompositePassColorKeyerCustomization::GetRGBWeightMax, InPropertyHandle)
+			.MaxSliderValue(this, &FCompositePassColorKeyerCustomization::GetRGBWeightMax, InPropertyHandle)
+			.LabelPadding(FMargin(3.f))
+			.LabelLocation(SNumericEntryBox<float>::ELabelLocation::Inside)
+			.Label()
+			[
+				SNumericEntryBox<float>::BuildNarrowColorLabel(InLabelColor)
+			];
+	};
+	
+	return SNew(SHorizontalBox)
+		+SHorizontalBox::Slot()
+		[
+			CreateNumericEntryBox(RedWeightHandle, AxisDisplayInfo::GetAxisColor(EAxisList::X))
+		]
+
+		+SHorizontalBox::Slot()
+		[
+			CreateNumericEntryBox(GreenWeightHandle, AxisDisplayInfo::GetAxisColor(EAxisList::Y))
+		]
+
+		+SHorizontalBox::Slot()
+		[
+			CreateNumericEntryBox(BlueWeightHandle, AxisDisplayInfo::GetAxisColor(EAxisList::Z))
+		];
+}
+
+TSharedRef<SWidget> FCompositePassColorKeyerCustomization::CreateCaptureButtonWidget()
+{
+	return SNew(SDelayedActionButton)
+		.IsEnabled(this, &FCompositePassColorKeyerCustomization::IsCaptureButtonEnabled)
+		.ButtonContent()
+		[
+			SNew(SBorder)
+			.HAlign(HAlign_Center)
+			.VAlign(VAlign_Center)
+			.BorderImage(FAppStyle::GetBrush("NoBorder"))
+			[
+				SNew(SHorizontalBox)
+
+				+SHorizontalBox::Slot()
+				.AutoWidth()
+				.VAlign(VAlign_Center)
+				.Padding(0.0f, 0.0f, 2.0f, 0.0f)
+				[
+					SNew(SImage)
+					.Image(FCompositeEditorStyle::Get().GetBrush("CompositeEditor.CaptureCleanPlate"))
+				]
+
+				+SHorizontalBox::Slot()
+				.AutoWidth()
+				.VAlign(VAlign_Center)
+				[
+					SNew(STextBlock)
+					.Text(LOCTEXT("CaptureButtonLabel", "Capture Clean Plate Texture"))
+					.Font(IDetailLayoutBuilder::GetDetailFont())
+				]
+			]
+		]
+		.CancelContent()
+		[
+			SNew(SBorder)
+			.HAlign(HAlign_Center)
+			.VAlign(VAlign_Center)
+			.Padding(0.0f, 2.0f)
+			.BorderImage(FAppStyle::GetBrush("NoBorder"))
+			[
+				SNew(SHorizontalBox)
+
+				+SHorizontalBox::Slot()
+				.AutoWidth()
+				.VAlign(VAlign_Center)
+				.Padding(0.0f, 0.0f, 2.0f, 0.0f)
+				[
+					SNew(SImage)
+					.Image(FCompositeEditorStyle::Get().GetBrush("CompositeEditor.CancelCapture"))
+				]
+
+				+SHorizontalBox::Slot()
+				.AutoWidth()
+				.VAlign(VAlign_Center)
+				[
+					SNew(STextBlock)
+					.Text(LOCTEXT("CancelCaptureButtonLabel", "Cancel Capture"))
+					.Font(IDetailLayoutBuilder::GetDetailFont())
+				]
+			]
+		]
+		.Delay(this, &FCompositePassColorKeyerCustomization::GetCaptureCountdown)
+		.OnActivateAction(this, &FCompositePassColorKeyerCustomization::CaptureMediaTexture);
+}
+
+float FCompositePassColorKeyerCustomization::GetCaptureCountdown() const
+{
+	if (CountdownHandle.IsValid())
+	{
+		float Countdown;
+		FPropertyAccess::Result Result = CountdownHandle->GetValue(Countdown);
+		if (Result == FPropertyAccess::Success)
+		{
+			return Countdown;
+		}
+	}
+
+	return 0.0;
+}
+
+TOptional<float> FCompositePassColorKeyerCustomization::GetRGBWeightValue(TSharedPtr<IPropertyHandle> InPropertyHandle) const
+{
+	if (InPropertyHandle.IsValid())
+	{
+		float Value;
+		FPropertyAccess::Result Result = InPropertyHandle->GetValue(Value);
+		if (Result == FPropertyAccess::Success)
+		{
+			return Value;
+		}
+	}
+
+	return TOptional<float>();
+}
+
+void FCompositePassColorKeyerCustomization::ChangeRGBWeightValue(float InNewValue, TSharedPtr<IPropertyHandle> InPropertyHandle)
+{
+	if (InPropertyHandle.IsValid())
+	{
+		InPropertyHandle->SetValue(InNewValue, EPropertyValueSetFlags::InteractiveChange);
+	}
+}
+
+void FCompositePassColorKeyerCustomization::CommitRGBWeightValue(float InNewValue, ETextCommit::Type InCommitType, TSharedPtr<IPropertyHandle> InPropertyHandle)
+{
+	if (InPropertyHandle.IsValid())
+	{
+		InPropertyHandle->SetValue(InNewValue);
+	}
+}
+
+bool FCompositePassColorKeyerCustomization::IsRGBWeightEditable(TSharedPtr<IPropertyHandle> InPropertyHandle) const
+{
+	if (InPropertyHandle.IsValid())
+	{
+		return InPropertyHandle->IsEditable();
+	}
+
+	return false;
+}
+
+TOptional<float> FCompositePassColorKeyerCustomization::GetRGBWeightMin(TSharedPtr<IPropertyHandle> InPropertyHandle) const
+{
+	TNumericPropertyParams<float> NumericPropertyParams(InPropertyHandle->GetProperty(), nullptr);
+	return NumericPropertyParams.MinSliderValue;
+}
+
+TOptional<float> FCompositePassColorKeyerCustomization::GetRGBWeightMax(TSharedPtr<IPropertyHandle> InPropertyHandle) const
+{
+	TNumericPropertyParams<float> NumericPropertyParams(InPropertyHandle->GetProperty(), nullptr);
+	return NumericPropertyParams.MaxSliderValue;
+}
+
+EVisibility FCompositePassColorKeyerCustomization::GetCaptureButtonVisibility() const
+{
+	if (KeyerSourceHandle.IsValid() && KeyerSourceHandle->IsValidHandle())
+	{
+		uint8 KeyerSource;
+		if (KeyerSourceHandle->GetValue(KeyerSource) == FPropertyAccess::Success)
+		{
+			return KeyerSource == (uint8)ECompositeKeyerSource::CleanPlate ? EVisibility::Visible : EVisibility::Collapsed;
+		}
+	}
+
+	return EVisibility::Collapsed;
+}
+
+bool FCompositePassColorKeyerCustomization::IsCaptureButtonEnabled() const
+{
+	TSharedPtr<IDetailLayoutBuilder> DetailBuilder = CachedDetailBuilder.Pin();
+	if (!DetailBuilder.IsValid())
+	{
+		return false;
+	}
+
+	TArray<TWeakObjectPtr<UCompositePassColorKeyer>> Keyers = DetailBuilder->GetObjectsOfTypeBeingCustomized<UCompositePassColorKeyer>();
+	if (Keyers.Num() != 1)
+	{
+		return false;
+	}
+
+	TStrongObjectPtr<UCompositePassColorKeyer> PinnedKeyer = Keyers[0].Pin();
+	if (!PinnedKeyer.IsValid())
+	{
+		return false;
+	}
+	
+	UCompositeLayerPlate* ParentPlate = PinnedKeyer->GetTypedOuter<UCompositeLayerPlate>();
+	if (!ParentPlate)
+	{
+		return false;
+	}
+
+	UMediaTexture* MediaTexture = Cast<UMediaTexture>(ParentPlate->Texture);
+	if (!MediaTexture)
+	{
+		return false;
+	}
+	
+	return true;
+}
+
+#undef LOCTEXT_NAMESPACE
diff --git a/Engine/Plugins/Compositing/Composite/Source/CompositeEditor/Private/Customizations/CompositePassColorKeyerCustomization.h b/Engine/Plugins/Compositing/Composite/Source/CompositeEditor/Private/Customizations/CompositePassColorKeyerCustomization.h
new file mode 100644
--- /dev/null
+++ b/Engine/Plugins/Compositing/Composite/Source/CompositeEditor/Private/Customizations/CompositePassColorKeyerCustomization.h
@@ -0,0 +1,82 @@
+﻿// Copyright Epic Games, Inc. All Rights Reserved.
+
+#pragma once
+
+#include "CompositePassBaseCustomization.h"
+#include "Layout/Visibility.h"
+#include "Types/SlateEnums.h"
+#include "UObject/StrongObjectPtr.h"
+
+class UMaterialInstanceDynamic;
+class IPropertyHandle;
+class SWidget;
+
+/**
+ * Customization for UCompositePassColorKeyer, which reorganizes some properties, and adds a 'Capture Clean Plate' button for simple clean plate capturing
+ */
+class FCompositePassColorKeyerCustomization : public FCompositePassBaseCustomization
+{
+public:
+	static TSharedRef<IDetailCustomization> MakeInstance();
+
+	virtual void CustomizeDetails(const TSharedPtr<IDetailLayoutBuilder>& DetailBuilder) override;
+	virtual void CustomizeDetails(IDetailLayoutBuilder& DetailBuilder) override;
+
+private:
+	void CaptureMediaTexture();
+
+	/** Creates the custom RGB single row widget, which contains color-coded numeric spinners for each component */
+	TSharedRef<SWidget> CreateRGBNumericEntryWidget();
+	
+	/** Creates the clean plate capture button widget */
+	TSharedRef<SWidget> CreateCaptureButtonWidget();
+
+	/** Gets the capture countdown property value on the UCompositePassColorKeyer objects being customized */
+	float GetCaptureCountdown() const;
+
+	/** Gets the RGB weight value for the specified component of the RGB property */
+	TOptional<float> GetRGBWeightValue(TSharedPtr<IPropertyHandle> InPropertyHandle) const;
+
+	/** Sets the RGB weight value for the specified component of the RGB property */
+	void ChangeRGBWeightValue(float InNewValue, TSharedPtr<IPropertyHandle> InPropertyHandle);
+
+	/** Commits the RGB weight value for the specified component of the RGB property */
+	void CommitRGBWeightValue(float InNewValue, ETextCommit::Type InCommitType, TSharedPtr<IPropertyHandle> InPropertyHandle);
+
+	/** Gets whether the RGB weight is editable for the  specified component of the RGB property */
+	bool IsRGBWeightEditable(TSharedPtr<IPropertyHandle> InPropertyHandle) const;
+
+	/** Gets the minimum RGB weight from the property metadata for the  specified component of the RGB property */
+	TOptional<float> GetRGBWeightMin(TSharedPtr<IPropertyHandle> InPropertyHandle) const;
+
+	/** Gets the maximum RGB weight from the property metadata for the  specified component of the RGB property */
+	TOptional<float> GetRGBWeightMax(TSharedPtr<IPropertyHandle> InPropertyHandle) const;
+
+	/** Gets the visibility state of the 'Capture Clean Plate' button */
+	EVisibility GetCaptureButtonVisibility() const;
+
+	/** Gets whether the 'Capture Clean Plate' button is enabled or not */
+	bool IsCaptureButtonEnabled() const;
+	
+private:
+	/** Cached detail layout builder pointer, used to manually refresh the details panel */
+	TWeakPtr<IDetailLayoutBuilder> CachedDetailBuilder;
+
+	/** Property handle for UCompositePassColorKeyer.KeyerSource property */
+	TSharedPtr<IPropertyHandle> KeyerSourceHandle;
+
+	/** Property handle for UCompositePassKeyer.CleanPlateCountdown property */
+	TSharedPtr<IPropertyHandle> CountdownHandle;
+
+	/** Property handle for R component of the UCompositePassKeyer.RGBWeight property */
+	TSharedPtr<IPropertyHandle> RedWeightHandle;
+
+	/** Property handle for G component of the UCompositePassKeyer.RGBWeight property */
+	TSharedPtr<IPropertyHandle> GreenWeightHandle;
+
+	/** Property handle for B component of the UCompositePassKeyer.RGBWeight property */
+	TSharedPtr<IPropertyHandle> BlueWeightHandle;
+
+	/** Transient dynamic material instance used to capture a clean plate texture from the keyer's media source */
+	TStrongObjectPtr<UMaterialInstanceDynamic> CaptureCleanPlateMID;
+};
diff --git a/Engine/Source/Runtime/Core/Public/UObject/FortniteMainBranchObjectVersions.inl b/Engine/Source/Runtime/Core/Public/UObject/FortniteMainBranchObjectVersions.inl
--- a/Engine/Source/Runtime/Core/Public/UObject/FortniteMainBranchObjectVersions.inl
+++ b/Engine/Source/Runtime/Core/Public/UObject/FortniteMainBranchObjectVersions.inl
@@ -745,5 +745,8 @@ FORTNITE_MAIN_VERSION(PolyBezierSplineAnchorTripleLayout, 243)
 // Repairing verse data where we introduced -0.0 value when we flipped the coordinate system
 FORTNITE_MAIN_VERSION(VerseRepairDataWithMinusZeroValue, 244)
 
+// Add new enum property to composite (plugin) UCompositePassColorKeyer to control keying settings
+FORTNITE_MAIN_VERSION(CompositePassKeyerSource, 245)
+
 // -----<new versions can be added above this line>-------------------------------------------------
 #undef FORTNITE_MAIN_VERSION
