Composite: Fixed missing dependency to EditorScriptingUtilities since it is now used by CompositeEditor module.

#Synced-CL 49870018

diff --git a/Engine/Plugins/Compositing/Composite/Composite.uplugin b/Engine/Plugins/Compositing/Composite/Composite.uplugin
--- a/Engine/Plugins/Compositing/Composite/Composite.uplugin
+++ b/Engine/Plugins/Compositing/Composite/Composite.uplugin
@@ -66,6 +66,11 @@
 			"Name": "ColorGrading",
 			"TargetAllowList": [ "Editor" ],
 			"Enabled": true
+		},
+		{
+			"Name": "EditorScriptingUtilities",
+			"Enabled": true,
+			"TargetAllowList": [ "Editor" ]
 		}
 	]
 }
\ No newline at end of file
