Composure: Default-disable main view resolution to avoid viewport resolution affecting performance.

* This follows engine separation of main view camera & resolution at CL 49938137.

#Synced-CL 49904425

diff --git a/Engine/Plugins/Compositing/Composite/Source/Composite/Private/Components/CompositeSceneCapture2DComponent.cpp b/Engine/Plugins/Compositing/Composite/Source/Composite/Private/Components/CompositeSceneCapture2DComponent.cpp
--- a/Engine/Plugins/Compositing/Composite/Source/Composite/Private/Components/CompositeSceneCapture2DComponent.cpp
+++ b/Engine/Plugins/Compositing/Composite/Source/Composite/Private/Components/CompositeSceneCapture2DComponent.cpp
@@ -23,7 +23,7 @@ UCompositeSceneCapture2DComponent::UCompositeSceneCapture2DComponent(const FObje
 	CaptureSource = ESceneCaptureSource::SCS_FinalColorHDR;
 	bUseRayTracingIfEnabled = true;
 	bMainViewFamily = true;
-	bMainViewResolution = true;
+	bMainViewResolution = false;
 	bMainViewCamera = true;
 	// Inherit main view camera post-process settings to match main view DoF for example.
 	bInheritMainViewCameraPostProcessSettings = true;
