Composure: Adding fast gaussian blur pass using weighted bilinear taps.

#Synced-CL 49904425

diff --git a/Engine/Plugins/Compositing/Composite/Shaders/Private/CompositeBlur.usf b/Engine/Plugins/Compositing/Composite/Shaders/Private/CompositeBlur.usf
new file mode 100644
--- /dev/null
+++ b/Engine/Plugins/Compositing/Composite/Shaders/Private/CompositeBlur.usf
@@ -0,0 +1,72 @@
+// Copyright Epic Games, Inc. All Rights Reserved.
+#include "/Engine/Private/Common.ush"
+#include "/Engine/Public/Platform.ush"
+
+// See https://www.rastergrid.com/blog/2010/09/efficient-gaussian-blur-with-linear-sampling/
+
+Texture2D InputTexture;
+SamplerState InputSampler;
+float2 Dir;
+float Radius;
+float2 InvTexSize;
+
+float NormalDistribution1D(float X)
+{
+    float X2 = PI * max(X, 0.0);
+    return exp(-0.5 * X2 * X2);
+}
+
+float4 GaussianBlur(float2 UV)
+{
+    float2 Scale = Dir * InvTexSize;
+    float4 Accum = 0.0;
+    float Total = 0.0;
+    float InvRadius = 1.0 / max(Radius, 1e-6);
+
+    // Handle center first
+    Accum += InputTexture.SampleLevel(InputSampler, UV, 0.0);
+    Total += 1.0;
+
+    // Determine how many pairs we can process symmetrically with bilinear taps
+    int NumPairs = (int)floor(Radius / 2.0);
+
+    // Want to replace two samples (Sample(d1) * w1) + (Sample(d2) * 22) with one Sample(O) * (w1 + w2).
+    // Per article, offset is the weighted average O = ((d1 * w1) + (d2 * w2)) / (w1 + w2).
+    // By replacing d2 with (d1 - 1) and simplifying, we get O = d1 - w2 / (w1 + w2).
+    LOOP
+    for (int Index = 0; Index < NumPairs; ++Index)
+    {
+        // Calculate distances for the two texels in this pair
+        float D1 = Radius - (Index * 2);
+        float D2 = D1 - 1.0f;
+
+        float W1 = NormalDistribution1D(D1 * InvRadius);
+        float W2 = NormalDistribution1D(D2 * InvRadius);
+        
+        float CombinedWeight = W1 + W2;
+        float Offset = D1 - W2 / CombinedWeight;
+        float2 UVOffset = Scale * Offset;
+        
+        Accum += InputTexture.SampleLevel(InputSampler, UV + UVOffset, 0.0) * CombinedWeight;
+        Accum += InputTexture.SampleLevel(InputSampler, UV - UVOffset, 0.0) * CombinedWeight; 
+        Total += CombinedWeight * 2.0;
+    }
+    
+    // Handle the leftover pixel if Radius is odd
+    if ((int)Radius % 2 != 0 && Radius >= 1.0)
+    {
+        float W = NormalDistribution1D(1.0 * InvRadius);
+        float2 UVOffset = Scale * 1.0;
+       
+        Accum += InputTexture.SampleLevel(InputSampler, UV + UVOffset, 0.0) * W;
+        Accum += InputTexture.SampleLevel(InputSampler, UV - UVOffset, 0.0) * W;
+        Total += W * 2.0;
+    }
+    
+    return Accum / max(Total, 1e-6);
+}
+
+float4 MainPS(float2 UV : TEXCOORD0) : SV_Target
+{
+    return GaussianBlur(UV);
+}
\ No newline at end of file
diff --git a/Engine/Plugins/Compositing/Composite/Source/Composite/Private/Passes/CompositePassBlur.cpp b/Engine/Plugins/Compositing/Composite/Source/Composite/Private/Passes/CompositePassBlur.cpp
new file mode 100644
--- /dev/null
+++ b/Engine/Plugins/Compositing/Composite/Source/Composite/Private/Passes/CompositePassBlur.cpp
@@ -0,0 +1,138 @@
+// Copyright Epic Games, Inc. All Rights Reserved.
+
+#include "Passes/CompositePassBlur.h"
+
+#include "RenderGraphBuilder.h"
+#include "RenderGraphUtils.h"
+#include "SceneView.h"
+#include "ScreenPass.h"
+
+DECLARE_GPU_STAT_NAMED(FCompositeBlur, TEXT("Composite.Blur"));
+
+class FCompositePassBlurShader : public FGlobalShader
+{
+	DECLARE_GLOBAL_SHADER(FCompositePassBlurShader);
+	SHADER_USE_PARAMETER_STRUCT(FCompositePassBlurShader, FGlobalShader);
+
+    BEGIN_SHADER_PARAMETER_STRUCT(FParameters, )
+    	SHADER_PARAMETER_RDG_TEXTURE(Texture2D, InputTexture)
+		SHADER_PARAMETER_SAMPLER(SamplerState, InputSampler)
+        SHADER_PARAMETER(FVector2f, InvTexSize)
+		SHADER_PARAMETER(FVector2f, Dir)
+		SHADER_PARAMETER(float, Radius)
+		RENDER_TARGET_BINDING_SLOTS()
+	END_SHADER_PARAMETER_STRUCT()
+};
+IMPLEMENT_GLOBAL_SHADER(FCompositePassBlurShader, "/Plugin/Composite/Private/CompositeBlur.usf", "MainPS", SF_Pixel);
+
+
+namespace UE
+{
+	namespace Composite
+	{
+		using namespace CompositeCore;
+
+		class FCompositePassBlurProxy : public FCompositeCorePassProxy
+		{
+		public:
+			IMPLEMENT_COMPOSITE_PASS(FCompositePassBlurProxy);
+
+			using FCompositeCorePassProxy::FCompositeCorePassProxy;
+
+			FPassTexture Add(FRDGBuilder& GraphBuilder, const FSceneView& InView, const FPassInputArray& Inputs, const FPassContext& PassContext) const override
+			{
+				RDG_EVENT_SCOPE_STAT(GraphBuilder, FCompositeBlur, "Composite.Blur");
+
+				check(ValidateInputs(Inputs));
+
+				const UE::CompositeCore::FResourceMetadata& Metadata = Inputs[0].Metadata;
+				const FScreenPassTexture& InputTexture = Inputs[0].Texture;
+				FScreenPassRenderTarget IntermediateOutput = CreateOutputRenderTarget(GraphBuilder, InView, PassContext.OutputViewRect, InputTexture.Texture->Desc, TEXT("CompositePassBlurIntermediate"));
+				FScreenPassRenderTarget Output = Inputs.OverrideOutput;
+
+				if (!Output.IsValid())
+				{
+					Output = CreateOutputRenderTarget(GraphBuilder, InView, PassContext.OutputViewRect, InputTexture.Texture->Desc, TEXT("CompositePassBlurOutput"));
+				}
+
+				FGlobalShaderMap* GlobalShaderMap = GetGlobalShaderMap(InView.GetFeatureLevel());
+				
+				{
+					FScreenPassTextureViewport Viewport = FScreenPassTextureViewport(IntermediateOutput);
+					FCompositePassBlurShader::FParameters* Parameters = GraphBuilder.AllocParameters<FCompositePassBlurShader::FParameters>();
+					Parameters->InputTexture = InputTexture.Texture;
+					Parameters->InputSampler = TStaticSamplerState<SF_Bilinear, AM_Clamp, AM_Clamp, AM_Clamp>::GetRHI();
+					const FIntVector InputSize = InputTexture.Texture->Desc.GetSize();
+					Parameters->InvTexSize = FVector2f::One() / FVector2f(FMath::Max(1.0f, (float)InputSize.X), FMath::Max(1.0f, (float)InputSize.Y));
+					Parameters->Dir = FVector2f(0, 1);
+					Parameters->Radius = Radius;
+					Parameters->RenderTargets[0] = IntermediateOutput.GetRenderTargetBinding();
+
+					TShaderMapRef<FCompositePassBlurShader> PixelShader(GlobalShaderMap);
+					AddDrawScreenPass(
+						GraphBuilder,
+						RDG_EVENT_NAME("Composite.Blur.Vertical (%dx%d)", Viewport.Extent.X, Viewport.Extent.Y),
+						InView,
+						Viewport,
+						Viewport,
+						PixelShader,
+						Parameters
+					);
+				}
+
+				{
+					FScreenPassTextureViewport Viewport = FScreenPassTextureViewport(Output);
+					FCompositePassBlurShader::FParameters* Parameters = GraphBuilder.AllocParameters<FCompositePassBlurShader::FParameters>();
+					Parameters->InputTexture = IntermediateOutput.Texture;
+					Parameters->InputSampler = TStaticSamplerState<SF_Bilinear, AM_Clamp, AM_Clamp, AM_Clamp>::GetRHI();
+					const FIntVector InputSize = IntermediateOutput.Texture->Desc.GetSize();
+					Parameters->InvTexSize = FVector2f(1.0f / FMath::Max(1.0f, (float)InputSize.X), 1.0f / FMath::Max(1.0f, (float)InputSize.Y));
+					Parameters->Dir = FVector2f(1, 0);
+					Parameters->Radius = Radius;
+					Parameters->RenderTargets[0] = Output.GetRenderTargetBinding();
+
+					TShaderMapRef<FCompositePassBlurShader> PixelShader(GlobalShaderMap);
+					AddDrawScreenPass(
+						GraphBuilder,
+						RDG_EVENT_NAME("Composite.Blur.Horizontal (%dx%d)", Viewport.Extent.X, Viewport.Extent.Y),
+						InView,
+						Viewport,
+						Viewport,
+						PixelShader,
+						Parameters
+					);
+				}
+
+				return UE::CompositeCore::FPassTexture{ MoveTemp(Output), Metadata };
+			}
+
+			int32 Radius = 1;
+		};
+	}
+}
+
+UCompositePassBlur::UCompositePassBlur(const FObjectInitializer& ObjectInitializer)
+	: Super(ObjectInitializer)
+	, Method(ECompositePassBlurMethod::Gaussian)
+    , Radius(1)
+{
+}
+
+UCompositePassBlur::~UCompositePassBlur() = default;
+
+bool UCompositePassBlur::IsActive() const
+{
+	return Super::IsActive() && Radius > 1;
+}
+
+bool UCompositePassBlur::GetProxy(const UE::CompositeCore::FPassInputDecl& InputDecl, FCompositeTraversalContext& InContext, FSceneRenderingBulkObjectAllocator& InFrameAllocator, FCompositeCorePassProxy*& OutProxy) const
+{
+	using namespace UE::Composite;
+
+	FCompositePassBlurProxy* Proxy = InFrameAllocator.Create<FCompositePassBlurProxy>(UE::CompositeCore::FPassInputDeclArray{ InputDecl });
+	Proxy->Radius = Radius;
+
+	OutProxy = Proxy;
+	return true;
+}
+
diff --git a/Engine/Plugins/Compositing/Composite/Source/Composite/Public/Passes/CompositePassBlur.h b/Engine/Plugins/Compositing/Composite/Source/Composite/Public/Passes/CompositePassBlur.h
new file mode 100644
--- /dev/null
+++ b/Engine/Plugins/Compositing/Composite/Source/Composite/Public/Passes/CompositePassBlur.h
@@ -0,0 +1,45 @@
+// Copyright Epic Games, Inc. All Rights Reserved.
+
+#pragma once
+
+#include "CompositePassBase.h"
+
+#include "CompositePassBlur.generated.h"
+
+#define UE_API COMPOSITE_API
+
+/** Blur method type. */
+UENUM(BlueprintType)
+enum class ECompositePassBlurMethod : uint8
+{
+	Gaussian,
+};
+
+/** Blur composite pass. */
+UCLASS(MinimalAPI, BlueprintType, Blueprintable, EditInlineNew, CollapseCategories, meta = (DisplayName = "Blur Pass"))
+class UCompositePassBlur : public UCompositePassBase
+{
+	GENERATED_BODY()
+
+public:
+	/** Constructor */
+	UE_API UCompositePassBlur(const FObjectInitializer& ObjectInitializer);
+
+	/** Destructor */
+	UE_API ~UCompositePassBlur();
+
+	UE_API virtual bool IsActive() const override;
+
+	UE_API virtual bool GetProxy(const UE::CompositeCore::FPassInputDecl& InputDecl, FCompositeTraversalContext& InContext, FSceneRenderingBulkObjectAllocator& InFrameAllocator, FCompositeCorePassProxy*& OutProxy) const override;
+
+public:
+    /** Blur pass method. */
+    UPROPERTY(EditAnywhere, BlueprintReadWrite, Category = "")
+    ECompositePassBlurMethod Method;
+
+	/** Size of the blur in pixels. */
+	UPROPERTY(EditAnywhere, BlueprintReadWrite, Category = "", meta = (ClampMin = "1", UIMin = "1", UIMax = "64"))
+	int32 Radius;
+};
+
+#undef UE_API
