Composite: Added Color Grading drawer integration for Single Light Shadow pass color grading settings

#Synced-CL 49893950

diff --git a/Engine/Plugins/Compositing/Composite/Source/Composite/Public/CompositeActor.h b/Engine/Plugins/Compositing/Composite/Source/Composite/Public/CompositeActor.h
--- a/Engine/Plugins/Compositing/Composite/Source/Composite/Public/CompositeActor.h
+++ b/Engine/Plugins/Compositing/Composite/Source/Composite/Public/CompositeActor.h
@@ -210,6 +210,7 @@ private:
 	friend class SCompositePanelLayerTree;
 	friend class FCompositeActorCustomization;
 	friend class FCompositeActorPanelDetailCustomization;
+	friend struct FColorGradingHierarchyConfig_Composite;
 	friend class UCompositeActorFactory;
 };
 
diff --git a/Engine/Plugins/Compositing/Composite/Source/CompositeEditor/Private/ColorGradingDrawer/ColorGradingDataModelGenerator_Composite.cpp b/Engine/Plugins/Compositing/Composite/Source/CompositeEditor/Private/ColorGradingDrawer/ColorGradingDataModelGenerator_Composite.cpp
--- a/Engine/Plugins/Compositing/Composite/Source/CompositeEditor/Private/ColorGradingDrawer/ColorGradingDataModelGenerator_Composite.cpp
+++ b/Engine/Plugins/Compositing/Composite/Source/CompositeEditor/Private/ColorGradingDrawer/ColorGradingDataModelGenerator_Composite.cpp
@@ -7,75 +7,118 @@
 #include "IDetailCustomization.h"
 #include "IDetailTreeNode.h"
 #include "IPropertyRowGenerator.h"
+#include "Layers/CompositeLayerSingleLightShadow.h"
 #include "Passes/CompositePassColorGrading.h"
 #include "Widgets/Images/SImage.h"
 #include "Widgets/Text/STextBlock.h"
 
 #define LOCTEXT_NAMESPACE "FColorGradingDataModelGenerator_Composite"
 
-class FCompositePassColorGradingCustomization : public IDetailCustomization
+namespace UE::Composite::ColorGrading
 {
-public:
-	virtual void CustomizeDetails(IDetailLayoutBuilder& DetailBuilder) override
+	class FCompositePassColorGradingCustomization : public IDetailCustomization
 	{
-		TArray<FName> Categories;
-		DetailBuilder.GetCategoryNames(Categories);
-
-		// Hide all default categories to avoid duplicate handles and unnecessary customization
-		for (const FName& Category : Categories)
+	public:
+		virtual void CustomizeDetails(IDetailLayoutBuilder& DetailBuilder) override
 		{
-			DetailBuilder.HideCategory(Category);
+			TArray<FName> Categories;
+			DetailBuilder.GetCategoryNames(Categories);
+
+			// Hide all default categories to avoid duplicate handles and unnecessary customization
+			for (const FName& Category : Categories)
+			{
+				DetailBuilder.HideCategory(Category);
+			}
+		
+			IDetailCategoryBuilder& ColorGradingHeaderCategory = DetailBuilder.EditCategory(TEXT("ColorGradingHeader"));
+			ColorGradingHeaderCategory.AddProperty(DetailBuilder.GetProperty(TEXT("bIsEnabled"), UCompositePassBase::StaticClass()));
+		
+			IDetailCategoryBuilder& ColorGradingElementsCategory = DetailBuilder.EditCategory(TEXT("ColorGradingElements"));
+
+			ColorGradingElementsCategory.AddProperty(DetailBuilder.GetProperty(GET_MEMBER_NAME_CHECKED(UCompositePassColorGrading, ColorGradingSettings.Global)));
+			ColorGradingElementsCategory.AddProperty(DetailBuilder.GetProperty(GET_MEMBER_NAME_CHECKED(UCompositePassColorGrading, ColorGradingSettings.Shadows)));
+			ColorGradingElementsCategory.AddProperty(DetailBuilder.GetProperty(GET_MEMBER_NAME_CHECKED(UCompositePassColorGrading, ColorGradingSettings.Midtones)));
+			ColorGradingElementsCategory.AddProperty(DetailBuilder.GetProperty(GET_MEMBER_NAME_CHECKED(UCompositePassColorGrading, ColorGradingSettings.Highlights)));
+
+			IDetailCategoryBuilder& TemperatureCategory = DetailBuilder.EditCategory(TEXT("Details_Temperature"), LOCTEXT("TemperatureCategory", "Temperature"));
+			TemperatureCategory.AddProperty(DetailBuilder.GetProperty(GET_MEMBER_NAME_CHECKED(UCompositePassColorGrading, TemperatureSettings.TemperatureType)));
+			TemperatureCategory.AddProperty(DetailBuilder.GetProperty(GET_MEMBER_NAME_CHECKED(UCompositePassColorGrading, TemperatureSettings.WhiteTemp)));
+			TemperatureCategory.AddProperty(DetailBuilder.GetProperty(GET_MEMBER_NAME_CHECKED(UCompositePassColorGrading, TemperatureSettings.WhiteTint)));
+
+			IDetailCategoryBuilder& ColorGradingCategory = DetailBuilder.EditCategory(TEXT("Details_ColorGrading"), LOCTEXT("ColorGradingCategory", "Color Grading"));
+			ColorGradingCategory.AddProperty(DetailBuilder.GetProperty(GET_MEMBER_NAME_CHECKED(UCompositePassColorGrading, ColorGradingSettings.ShadowsMax)));
+			ColorGradingCategory.AddProperty(DetailBuilder.GetProperty(GET_MEMBER_NAME_CHECKED(UCompositePassColorGrading, ColorGradingSettings.HighlightsMin)));
+			ColorGradingCategory.AddProperty(DetailBuilder.GetProperty(GET_MEMBER_NAME_CHECKED(UCompositePassColorGrading, ColorGradingSettings.HighlightsMax)));
 		}
+	};
+
+	class FCompositeLayerSingleLightShadowCustomization : public IDetailCustomization
+	{
+	public:
+		virtual void CustomizeDetails(IDetailLayoutBuilder& DetailBuilder) override
+		{
+			TArray<FName> Categories;
+			DetailBuilder.GetCategoryNames(Categories);
+
+			// Hide all default categories to avoid duplicate handles and unnecessary customization
+			for (const FName& Category : Categories)
+			{
+				DetailBuilder.HideCategory(Category);
+			}
 		
-		IDetailCategoryBuilder& ColorGradingHeaderCategory = DetailBuilder.EditCategory(TEXT("ColorGradingHeader"));
-		ColorGradingHeaderCategory.AddProperty(DetailBuilder.GetProperty(TEXT("bIsEnabled"), UCompositePassBase::StaticClass()));
+			IDetailCategoryBuilder& ColorGradingHeaderCategory = DetailBuilder.EditCategory(TEXT("ColorGradingHeader"));
+			ColorGradingHeaderCategory.AddProperty(DetailBuilder.GetProperty(GET_MEMBER_NAME_CHECKED(UCompositeLayerSingleLightShadow, ColorGrading.bEnableColorGrading)));
 		
-		IDetailCategoryBuilder& ColorGradingElementsCategory = DetailBuilder.EditCategory(TEXT("ColorGradingElements"));
-
-		ColorGradingElementsCategory.AddProperty(DetailBuilder.GetProperty(GET_MEMBER_NAME_CHECKED(UCompositePassColorGrading, ColorGradingSettings.Global)));
-		ColorGradingElementsCategory.AddProperty(DetailBuilder.GetProperty(GET_MEMBER_NAME_CHECKED(UCompositePassColorGrading, ColorGradingSettings.Shadows)));
-		ColorGradingElementsCategory.AddProperty(DetailBuilder.GetProperty(GET_MEMBER_NAME_CHECKED(UCompositePassColorGrading, ColorGradingSettings.Midtones)));
-		ColorGradingElementsCategory.AddProperty(DetailBuilder.GetProperty(GET_MEMBER_NAME_CHECKED(UCompositePassColorGrading, ColorGradingSettings.Highlights)));
-
-		IDetailCategoryBuilder& TemperatureCategory = DetailBuilder.EditCategory(TEXT("Details_Temperature"), LOCTEXT("TemperatureCategory", "Temperature"));
-		TemperatureCategory.AddProperty(DetailBuilder.GetProperty(GET_MEMBER_NAME_CHECKED(UCompositePassColorGrading, TemperatureSettings.TemperatureType)));
-		TemperatureCategory.AddProperty(DetailBuilder.GetProperty(GET_MEMBER_NAME_CHECKED(UCompositePassColorGrading, TemperatureSettings.WhiteTemp)));
-		TemperatureCategory.AddProperty(DetailBuilder.GetProperty(GET_MEMBER_NAME_CHECKED(UCompositePassColorGrading, TemperatureSettings.WhiteTint)));
-
-		IDetailCategoryBuilder& ColorGradingCategory = DetailBuilder.EditCategory(TEXT("Details_ColorGrading"), LOCTEXT("ColorGradingCategory", "Color Grading"));
-		ColorGradingCategory.AddProperty(DetailBuilder.GetProperty(GET_MEMBER_NAME_CHECKED(UCompositePassColorGrading, ColorGradingSettings.ShadowsMax)));
-		ColorGradingCategory.AddProperty(DetailBuilder.GetProperty(GET_MEMBER_NAME_CHECKED(UCompositePassColorGrading, ColorGradingSettings.HighlightsMin)));
-		ColorGradingCategory.AddProperty(DetailBuilder.GetProperty(GET_MEMBER_NAME_CHECKED(UCompositePassColorGrading, ColorGradingSettings.HighlightsMax)));
-	}
-};
+			IDetailCategoryBuilder& ColorGradingElementsCategory = DetailBuilder.EditCategory(TEXT("ColorGradingElements"));
+
+			ColorGradingElementsCategory.AddProperty(DetailBuilder.GetProperty(GET_MEMBER_NAME_CHECKED(UCompositeLayerSingleLightShadow, ColorGrading.ColorGradingSettings.Global)));
+			ColorGradingElementsCategory.AddProperty(DetailBuilder.GetProperty(GET_MEMBER_NAME_CHECKED(UCompositeLayerSingleLightShadow, ColorGrading.ColorGradingSettings.Shadows)));
+			ColorGradingElementsCategory.AddProperty(DetailBuilder.GetProperty(GET_MEMBER_NAME_CHECKED(UCompositeLayerSingleLightShadow, ColorGrading.ColorGradingSettings.Midtones)));
+			ColorGradingElementsCategory.AddProperty(DetailBuilder.GetProperty(GET_MEMBER_NAME_CHECKED(UCompositeLayerSingleLightShadow, ColorGrading.ColorGradingSettings.Highlights)));
+
+			IDetailCategoryBuilder& TemperatureCategory = DetailBuilder.EditCategory(TEXT("Details_Temperature"), LOCTEXT("TemperatureCategory", "Temperature"));
+			TemperatureCategory.AddProperty(DetailBuilder.GetProperty(GET_MEMBER_NAME_CHECKED(UCompositeLayerSingleLightShadow, ColorGrading.TemperatureSettings.TemperatureType)));
+			TemperatureCategory.AddProperty(DetailBuilder.GetProperty(GET_MEMBER_NAME_CHECKED(UCompositeLayerSingleLightShadow, ColorGrading.TemperatureSettings.WhiteTemp)));
+			TemperatureCategory.AddProperty(DetailBuilder.GetProperty(GET_MEMBER_NAME_CHECKED(UCompositeLayerSingleLightShadow, ColorGrading.TemperatureSettings.WhiteTint)));
+
+			IDetailCategoryBuilder& ColorGradingCategory = DetailBuilder.EditCategory(TEXT("Details_ColorGrading"), LOCTEXT("ColorGradingCategory", "Color Grading"));
+			ColorGradingCategory.AddProperty(DetailBuilder.GetProperty(GET_MEMBER_NAME_CHECKED(UCompositeLayerSingleLightShadow, ColorGrading.ColorGradingSettings.ShadowsMax)));
+			ColorGradingCategory.AddProperty(DetailBuilder.GetProperty(GET_MEMBER_NAME_CHECKED(UCompositeLayerSingleLightShadow, ColorGrading.ColorGradingSettings.HighlightsMin)));
+			ColorGradingCategory.AddProperty(DetailBuilder.GetProperty(GET_MEMBER_NAME_CHECKED(UCompositeLayerSingleLightShadow, ColorGrading.ColorGradingSettings.HighlightsMax)));
+		}
+	};
+}
 
 void FColorGradingDataModelGenerator_Composite::Initialize(const TSharedRef<FColorGradingEditorDataModel>& ColorGradingDataModel, const TSharedRef<IPropertyRowGenerator>& PropertyRowGenerator)
 {
 	PropertyRowGenerator->RegisterInstancedCustomPropertyLayout(UCompositePassColorGrading::StaticClass(), FOnGetDetailCustomizationInstance::CreateLambda([]
 	{
-		return MakeShared<FCompositePassColorGradingCustomization>();
+		return MakeShared<UE::Composite::ColorGrading::FCompositePassColorGradingCustomization>();
+	}));
+
+	PropertyRowGenerator->RegisterInstancedCustomPropertyLayout(UCompositeLayerSingleLightShadow::StaticClass(), FOnGetDetailCustomizationInstance::CreateLambda([]
+	{
+		return MakeShared<UE::Composite::ColorGrading::FCompositeLayerSingleLightShadowCustomization>();
 	}));
 }
 
 void FColorGradingDataModelGenerator_Composite::Destroy(const TSharedRef<FColorGradingEditorDataModel>& ColorGradingDataModel, const TSharedRef<IPropertyRowGenerator>& PropertyRowGenerator)
 {
 	PropertyRowGenerator->UnregisterInstancedCustomPropertyLayout(UCompositePassColorGrading::StaticClass());
+	PropertyRowGenerator->UnregisterInstancedCustomPropertyLayout(UCompositeLayerSingleLightShadow::StaticClass());
 }
 
 void FColorGradingDataModelGenerator_Composite::GenerateDataModel(IPropertyRowGenerator& PropertyRowGenerator, FColorGradingEditorDataModel& OutColorGradingDataModel)
 {
 	TArray<TWeakObjectPtr<UCompositePassColorGrading>> SelectedPasses;
 	const TArray<TWeakObjectPtr<UObject>>& SelectedObjects = PropertyRowGenerator.GetSelectedObjects();
-	for (const TWeakObjectPtr<UObject>& SelectedObject : SelectedObjects)
+
+	const bool bCanGenerateDataModel = SelectedObjects.ContainsByPredicate([](const TWeakObjectPtr<UObject>& Object)
 	{
-		if (SelectedObject.IsValid() && SelectedObject->IsA<UCompositePassColorGrading>())
-		{
-			TWeakObjectPtr<UCompositePassColorGrading> SelectedRootActor = CastChecked<UCompositePassColorGrading>(SelectedObject.Get());
-			SelectedPasses.Add(SelectedRootActor);
-		}
-	}
+		return Object.IsValid() && (Object->IsA<UCompositePassColorGrading>() || Object->IsA<UCompositeLayerSingleLightShadow>());
+	});
 
-	if (SelectedPasses.IsEmpty())
+	if (!bCanGenerateDataModel)
 	{
 		return;
 	}
diff --git a/Engine/Plugins/Compositing/Composite/Source/CompositeEditor/Private/ColorGradingDrawer/ColorGradingHierarchyConfig_Composite.cpp b/Engine/Plugins/Compositing/Composite/Source/CompositeEditor/Private/ColorGradingDrawer/ColorGradingHierarchyConfig_Composite.cpp
--- a/Engine/Plugins/Compositing/Composite/Source/CompositeEditor/Private/ColorGradingDrawer/ColorGradingHierarchyConfig_Composite.cpp
+++ b/Engine/Plugins/Compositing/Composite/Source/CompositeEditor/Private/ColorGradingDrawer/ColorGradingHierarchyConfig_Composite.cpp
@@ -5,6 +5,7 @@
 #include "CompositeActor.h"
 #include "CompositeEditorStyle.h"
 #include "Layers/CompositeLayerPlate.h"
+#include "Layers/CompositeLayerSingleLightShadow.h"
 #include "ObjectFilter/ObjectMixerEditorActorSubObject.h"
 #include "Passes/CompositePassColorGrading.h"
 
@@ -16,8 +17,16 @@ TArray<FObjectMixerEditorActorSubObject> FColorGradingHierarchyConfig_Composite:
 	{
 		for (const TObjectPtr<UCompositeLayerBase>& CompositeLayer : CompositeActor->GetCompositeLayers())
 		{
-			const UCompositeLayerPlate* PlateLayer = Cast<UCompositeLayerPlate>(CompositeLayer.Get());
-			if (IsValid(PlateLayer))
+			if (const UCompositeLayerSingleLightShadow* SingleLightShadowLayer = Cast<UCompositeLayerSingleLightShadow>(CompositeLayer))
+			{
+				FObjectMixerEditorActorSubObject SubObject;
+				SubObject.Object = CompositeLayer;
+				SubObject.DisplayName = FText::FromString(!CompositeLayer->Name.IsEmpty() ? CompositeLayer->Name : CompositeLayer->GetName());
+				SubObject.bIncludeObjectStack = true;
+				
+				SubObjects.Add(SubObject);
+			}
+			else if (const UCompositeLayerPlate* PlateLayer = Cast<UCompositeLayerPlate>(CompositeLayer.Get()))
 			{
 				for (const TObjectPtr<UCompositePassBase>& Pass : PlateLayer->MediaPasses)
 				{
@@ -60,6 +69,7 @@ TSet<FName> FColorGradingHierarchyConfig_Composite::GetPropertiesThatRequireList
 	return
 	{
 		GET_MEMBER_NAME_CHECKED(UCompositeLayerPlate, MediaPasses),
-		GET_MEMBER_NAME_CHECKED(UCompositeLayerPlate, LayerPasses)
+		GET_MEMBER_NAME_CHECKED(UCompositeLayerPlate, LayerPasses),
+		GET_MEMBER_NAME_CHECKED(ACompositeActor, CompositeLayers)
 	};
 }
diff --git a/Engine/Plugins/Compositing/Composite/Source/CompositeEditor/Private/ColorGradingDrawer/ColorGradingSelectionHandler_Composite.cpp b/Engine/Plugins/Compositing/Composite/Source/CompositeEditor/Private/ColorGradingDrawer/ColorGradingSelectionHandler_Composite.cpp
--- a/Engine/Plugins/Compositing/Composite/Source/CompositeEditor/Private/ColorGradingDrawer/ColorGradingSelectionHandler_Composite.cpp
+++ b/Engine/Plugins/Compositing/Composite/Source/CompositeEditor/Private/ColorGradingDrawer/ColorGradingSelectionHandler_Composite.cpp
@@ -3,6 +3,7 @@
 #include "ColorGradingSelectionHandler_Composite.h"
 
 #include "CompositeActor.h"
+#include "Layers/CompositeLayerSingleLightShadow.h"
 #include "Passes/CompositePassColorGrading.h"
 #include "UI/SCompositeEditorPanel.h"
 
@@ -39,7 +40,7 @@ TArray<UObject*> FColorGradingSelectionHandler_Composite::GetSelectedSubObjectsI
 	// Remove any non-color grading objects from the selection
 	SelectedObjects.RemoveAll([](const UObject* Object)
 	{
-		return !Object->IsA<UCompositePassColorGrading>();
+		return !Object->IsA<UCompositePassColorGrading>() && !Object->IsA<UCompositeLayerSingleLightShadow>();
 	});
 
 	return SelectedObjects;
diff --git a/Engine/Plugins/Compositing/Composite/Source/CompositeEditor/Private/CompositeEditorModule.cpp b/Engine/Plugins/Compositing/Composite/Source/CompositeEditor/Private/CompositeEditorModule.cpp
--- a/Engine/Plugins/Compositing/Composite/Source/CompositeEditor/Private/CompositeEditorModule.cpp
+++ b/Engine/Plugins/Compositing/Composite/Source/CompositeEditor/Private/CompositeEditorModule.cpp
@@ -55,6 +55,8 @@ void FCompositeEditorModule::StartupModule()
 
 	FColorGradingEditorDataModel::RegisterColorGradingDataModelGenerator<UCompositePassColorGrading>(
 		FGetDetailsDataModelGenerator::CreateStatic(&FColorGradingDataModelGenerator_Composite::MakeInstance));
+	FColorGradingEditorDataModel::RegisterColorGradingDataModelGenerator<UCompositeLayerSingleLightShadow>(
+		FGetDetailsDataModelGenerator::CreateStatic(&FColorGradingDataModelGenerator_Composite::MakeInstance));
 	
 	FColorGradingMixerObjectFilterRegistry::RegisterActorClassToPlace(ACompositeActor::StaticClass());
 
diff --git a/Engine/Plugins/Compositing/Composite/Source/CompositeEditor/Private/Customizations/CompositeLayerSingleLightShadowCustomization.cpp b/Engine/Plugins/Compositing/Composite/Source/CompositeEditor/Private/Customizations/CompositeLayerSingleLightShadowCustomization.cpp
--- a/Engine/Plugins/Compositing/Composite/Source/CompositeEditor/Private/Customizations/CompositeLayerSingleLightShadowCustomization.cpp
+++ b/Engine/Plugins/Compositing/Composite/Source/CompositeEditor/Private/Customizations/CompositeLayerSingleLightShadowCustomization.cpp
@@ -31,6 +31,10 @@ void FCompositeLayerSingleLightShadowCustomization::CustomizeDetails(IDetailLayo
 	TArray<TWeakObjectPtr<UCompositeLayerSingleLightShadow>> Objects = DetailLayout.GetObjectsOfTypeBeingCustomized<UCompositeLayerSingleLightShadow>();
 
 	IDetailCategoryBuilder& LayerCategory = DetailLayout.EditCategoryAllowNone(NAME_None);
+
+	// Mark ColorGrading as customized so that it can be placed below 'Shadow Casting Content' custom group
+	TSharedPtr<IPropertyHandle> ColorGradingHandle = DetailLayout.GetProperty(GET_MEMBER_NAME_CHECKED(UCompositeLayerSingleLightShadow, ColorGrading));
+	ColorGradingHandle->MarkHiddenByCustomization();
 	
 	// Add all default simple properties ahead of the custom group
 	TArray<TSharedRef<IPropertyHandle>> LayerProperties;
@@ -75,6 +79,8 @@ void FCompositeLayerSingleLightShadowCustomization::CustomizeDetails(IDetailLayo
 			.Font(IDetailLayoutBuilder::GetDetailFont())
 		];
 	}
+
+	LayerCategory.AddProperty(ColorGradingHandle);
 }
 
 void FCompositeLayerSingleLightShadowCustomization::OnLayoutSizeChanged()
