Composure: Fix to prevent addition of the same composite mesh on different plate layers.

* The plate layer post-edit change logic was previously only working on the regular details panel.
* Full implementation would require a chain event. Last modified index logic suffices for the base case.

#Synced-CL 50287290

diff --git a/Engine/Plugins/Compositing/Composite/Source/CompositeEditor/Private/UI/SCompositeActorPickerSceneOutliner.cpp b/Engine/Plugins/Compositing/Composite/Source/CompositeEditor/Private/UI/SCompositeActorPickerSceneOutliner.cpp
--- a/Engine/Plugins/Compositing/Composite/Source/CompositeEditor/Private/UI/SCompositeActorPickerSceneOutliner.cpp
+++ b/Engine/Plugins/Compositing/Composite/Source/CompositeEditor/Private/UI/SCompositeActorPickerSceneOutliner.cpp
@@ -187,11 +187,13 @@ void SCompositeActorPickerSceneOutliner::OnActorSelected(AActor* InActor, bool b
 	PinnedListOwner->Modify();
 
 	ActorListRef.NotifyPreEditChange();
+	int32 LastModifiedIndex = INDEX_NONE;
 
 	EPropertyChangeType::Type ChangeType = EPropertyChangeType::Unspecified;
 	if (bIsAdding)
 	{
 		ActorListRef.ActorList->Add(InActor);
+		LastModifiedIndex = ActorListRef.ActorList->Num() - 1;
 		ChangeType = EPropertyChangeType::ArrayAdd;
 	}
 	else if (bIsRemoving)
@@ -200,7 +202,7 @@ void SCompositeActorPickerSceneOutliner::OnActorSelected(AActor* InActor, bool b
 		ChangeType = EPropertyChangeType::ArrayRemove;
 	}
 
-	ActorListRef.NotifyPostEditChangeList(ChangeType);
+	ActorListRef.NotifyPostEditChangeList(ChangeType, LastModifiedIndex);
 	OnActorListChanged.ExecuteIfBound();
 }
 
diff --git a/Engine/Plugins/Compositing/Composite/Source/CompositeEditor/Private/UI/SCompositeActorPickerTable.cpp b/Engine/Plugins/Compositing/Composite/Source/CompositeEditor/Private/UI/SCompositeActorPickerTable.cpp
--- a/Engine/Plugins/Compositing/Composite/Source/CompositeEditor/Private/UI/SCompositeActorPickerTable.cpp
+++ b/Engine/Plugins/Compositing/Composite/Source/CompositeEditor/Private/UI/SCompositeActorPickerTable.cpp
@@ -184,6 +184,7 @@ private:
 		}
 
 		TSharedPtr<FScopedTransaction> AddActorsTransaction = nullptr;
+		int32 LastModifiedIndex = INDEX_NONE;
 
 		for (AActor* SelectedActor : SelectedActors)
 		{
@@ -201,11 +202,12 @@ private:
 			}
 
 			ActorListRef.ActorList->Add(SelectedActor);
+			LastModifiedIndex = ActorListRef.ActorList->Num() - 1;
 		}
 
 		if (AddActorsTransaction.IsValid())
 		{
-			ActorListRef.NotifyPostEditChangeList(EPropertyChangeType::ArrayAdd);
+			ActorListRef.NotifyPostEditChangeList(EPropertyChangeType::ArrayAdd, LastModifiedIndex);
 			OnActorListChanged.ExecuteIfBound();
 		}
 	}
@@ -360,7 +362,7 @@ void FCompositeActorPickerListRef::NotifyPreEditChange()
 	}
 }
 
-void FCompositeActorPickerListRef::NotifyPostEditChangeList(EPropertyChangeType::Type InChangeType)
+void FCompositeActorPickerListRef::NotifyPostEditChangeList(EPropertyChangeType::Type InChangeType, int32 LastModifiedIndex)
 {
 	if (!IsValid())
 	{
@@ -370,6 +372,14 @@ void FCompositeActorPickerListRef::NotifyPostEditChangeList(EPropertyChangeType:
 	if (TStrongObjectPtr<UObject> PinnedOwner = ActorListOwner.Pin())
 	{
 		FPropertyChangedEvent ChangedEvent(ActorListProperty, InChangeType);
+		
+		// TODO: Full implementation would require an FPropertyChangedChainEvent
+
+		TMap<FString, int32> ArrayIndexPerObject;
+		ArrayIndexPerObject.Add(ActorListPropertyName.ToString(), LastModifiedIndex);
+		ChangedEvent.ObjectIteratorIndex = 0;
+		ChangedEvent.SetArrayIndexPerObject(MakeArrayView(&ArrayIndexPerObject, 1));
+
 		PinnedOwner->PostEditChangeProperty(ChangedEvent);
 	}
 }
@@ -664,6 +674,7 @@ FReply SCompositeActorPickerTable::OnListDropped(const FGeometry& Geometry, cons
 	}
 			
 	TSharedPtr<FScopedTransaction> AddActorsTransaction = nullptr;
+	int32 LastModifiedIndex = INDEX_NONE;
 	for (AActor* DroppedActor : DroppedActors)
 	{
 		if (ActorListRef.ActorList->Contains(DroppedActor))
@@ -680,11 +691,12 @@ FReply SCompositeActorPickerTable::OnListDropped(const FGeometry& Geometry, cons
 		}
 
 		ActorListRef.ActorList->Add(DroppedActor);
+		LastModifiedIndex = ActorListRef.ActorList->Num() - 1;
 	}
 
 	if (AddActorsTransaction.IsValid())
 	{
-		ActorListRef.NotifyPostEditChangeList(EPropertyChangeType::ArrayAdd);
+		ActorListRef.NotifyPostEditChangeList(EPropertyChangeType::ArrayAdd, LastModifiedIndex);
 		FillActorList();
 	}
 
diff --git a/Engine/Plugins/Compositing/Composite/Source/CompositeEditor/Private/UI/SCompositeActorPickerTable.h b/Engine/Plugins/Compositing/Composite/Source/CompositeEditor/Private/UI/SCompositeActorPickerTable.h
--- a/Engine/Plugins/Compositing/Composite/Source/CompositeEditor/Private/UI/SCompositeActorPickerTable.h
+++ b/Engine/Plugins/Compositing/Composite/Source/CompositeEditor/Private/UI/SCompositeActorPickerTable.h
@@ -26,7 +26,7 @@ public:
 	void NotifyPreEditChange();
 
 	/** Raises the owner's PostEditchange event for the actor list property */
-	void NotifyPostEditChangeList(EPropertyChangeType::Type InChangeType);
+	void NotifyPostEditChangeList(EPropertyChangeType::Type InChangeType, int32 LastModifiedIndex = INDEX_NONE);
 
 	/** Finds the localized tooltip or native tooltip as a fallback. */
 	FText GetToolTipText(bool bShortTooltip = false) const;
