Composure: Added plate texture preview that is displayed when a keyer pass's key color picker or eye dropper is used, allowing users to select colors from the raw plate texture

#Synced-CL 50006453

diff --git a/Engine/Plugins/Compositing/Composite/Source/Composite/Composite.Build.cs b/Engine/Plugins/Compositing/Composite/Source/Composite/Composite.Build.cs
--- a/Engine/Plugins/Compositing/Composite/Source/Composite/Composite.Build.cs
+++ b/Engine/Plugins/Compositing/Composite/Source/Composite/Composite.Build.cs
@@ -41,6 +41,7 @@ public class Composite : ModuleRules
 			PrivateDependencyModuleNames.AddRange(
 				new string[]
 				{
+					"LevelEditor",
 					"UnrealEd",
 					"Slate",
 					"SlateCore",
diff --git a/Engine/Plugins/Compositing/Composite/Source/Composite/Private/Components/CompositePlateTexturePreviewComponent.cpp b/Engine/Plugins/Compositing/Composite/Source/Composite/Private/Components/CompositePlateTexturePreviewComponent.cpp
new file mode 100644
--- /dev/null
+++ b/Engine/Plugins/Compositing/Composite/Source/Composite/Private/Components/CompositePlateTexturePreviewComponent.cpp
@@ -0,0 +1,155 @@
+﻿// Copyright Epic Games, Inc. All Rights Reserved.
+
+#include "CompositePlateTexturePreviewComponent.h"
+
+#include "CompositeActor.h"
+#include "Brushes/SlateImageBrush.h"
+#include "Engine/Texture.h"
+#include "Layers/CompositeLayerPlate.h"
+#include "Passes/CompositePassColorKeyer.h"
+#include "Widgets/DeclarativeSyntaxSupport.h"
+#include "Widgets/Images/SImage.h"
+#include "Widgets/Layout/SScaleBox.h"
+
+#if WITH_EDITOR
+#include "Editor.h"
+#include "LevelEditorViewport.h"
+#include "Settings/LevelEditorViewportSettings.h"
+#include "SLevelViewport.h"
+#endif
+
+#if WITH_EDITOR
+bool UCompositePlateTexturePreviewComponent::GetEditorPreviewInfo(float DeltaTime, FMinimalViewInfo& ViewOut)
+{
+	if (TStrongObjectPtr<UCompositeLayerPlate> PinnedPlate = PlateToPreview.Pin())
+	{
+		if (UTexture* Texture = PinnedPlate->Texture)
+		{
+			// Editor preview widget uses ViewOut.AspectRatio to determine the size to make the preview widget, so set the apsect ratio
+			// to match the texture being previewed if there is one
+			ViewOut.bConstrainAspectRatio = true;
+			ViewOut.AspectRatio = float(Texture->GetSurfaceWidth()) / float(Texture->GetSurfaceHeight());
+		}
+	}
+
+	return true;
+}
+
+TSharedPtr<SWidget> UCompositePlateTexturePreviewComponent::GetCustomEditorPreviewWidget()
+{
+	TStrongObjectPtr<UCompositeLayerPlate> PinnedPlate = PlateToPreview.Pin();
+	if (!PinnedPlate.IsValid())
+	{
+		return nullptr;
+	}
+
+	UTexture* Texture = PinnedPlate->Texture;
+	if (!Texture)
+	{
+		return nullptr;
+	}
+
+	const int64 Width = Texture->GetSurfaceWidth();
+	const int64 Height = Texture->GetSurfaceHeight();
+
+	if (!PreviewBrush.IsValid())
+	{
+		PreviewBrush = MakeShared<FSlateImageBrush>(Texture, FVector2D(Width, Height), FLinearColor::White, ESlateBrushTileType::Both);
+	}
+	else
+	{
+		PreviewBrush->SetResourceObject(Texture);
+		PreviewBrush->SetImageSize(FVector2D(Width, Height));
+	}
+
+	return SNew(SScaleBox)
+		.Stretch(EStretch::ScaleToFit)
+		[
+			SNew(SImage)
+			.Image(PreviewBrush.Get())
+		];
+}
+
+bool UCompositePlateTexturePreviewComponent::ShowPreview(TObjectPtr<UCompositePassColorKeyer> InRequestingKeyer)
+{
+	Activate();
+	
+	UCompositeLayerPlate* Plate = InRequestingKeyer->GetTypedOuter<UCompositeLayerPlate>();
+	if (!Plate)
+	{
+		return false;
+	}
+
+	AActor* PlateOwner = Plate->GetTypedOuter<AActor>();
+	if (!PlateOwner || PlateOwner != GetOwner())
+	{
+		return false;
+	}
+
+	if (!Plate->Texture)
+	{
+		return false;
+	}
+
+	if (FLevelEditorViewportClient* ViewportClient = GetActiveViewportClient())
+	{
+		if (TSharedPtr<SLevelViewport> LevelViewport = StaticCastSharedPtr<SLevelViewport>(ViewportClient->GetEditorViewportWidget()))
+		{
+			PlateToPreview = Plate;
+			bShowPreview = true;
+
+			if (ULevelEditorViewportSettings* ViewportSettings = GetMutableDefault<ULevelEditorViewportSettings>())
+			{
+				CachedCameraPreviewSize = ViewportSettings->CameraPreviewSize;
+				ViewportSettings->CameraPreviewSize = InRequestingKeyer->PlateTexturePreviewSize;
+			}
+			
+			LevelViewport->SetActorAlwaysPreview(GetOwner(), true);
+			return true;
+		}
+	}
+	
+	return false;
+}
+
+void UCompositePlateTexturePreviewComponent::HidePreview()
+{
+	Deactivate();
+	
+	if (FLevelEditorViewportClient* ViewportClient = GetActiveViewportClient())
+	{
+		if (TSharedPtr<SLevelViewport> LevelViewport = StaticCastSharedPtr<SLevelViewport>(ViewportClient->GetEditorViewportWidget()))
+		{
+			LevelViewport->SetActorAlwaysPreview(GetOwner(), false);
+		}
+	}
+	
+	PlateToPreview = nullptr;
+	bShowPreview = false;
+
+	if (CachedCameraPreviewSize.IsSet())
+	{
+		if (ULevelEditorViewportSettings* ViewportSettings = GetMutableDefault<ULevelEditorViewportSettings>())
+		{
+			ViewportSettings->CameraPreviewSize = CachedCameraPreviewSize.GetValue();
+		}
+
+		CachedCameraPreviewSize.Reset();
+	}
+}
+
+FLevelEditorViewportClient* UCompositePlateTexturePreviewComponent::GetActiveViewportClient()
+{
+	if (const FViewport* ActiveViewport = GEditor ? GEditor->GetActiveViewport() : nullptr)
+	{
+		for (FLevelEditorViewportClient* LevelViewportClient : GEditor->GetLevelViewportClients())
+		{
+			if (LevelViewportClient->Viewport == ActiveViewport)
+			{
+				return LevelViewportClient;
+			}
+		}
+	}
+	return nullptr;
+}
+#endif
\ No newline at end of file
diff --git a/Engine/Plugins/Compositing/Composite/Source/Composite/Private/Components/CompositePlateTexturePreviewComponent.h b/Engine/Plugins/Compositing/Composite/Source/Composite/Private/Components/CompositePlateTexturePreviewComponent.h
new file mode 100644
--- /dev/null
+++ b/Engine/Plugins/Compositing/Composite/Source/Composite/Private/Components/CompositePlateTexturePreviewComponent.h
@@ -0,0 +1,60 @@
+﻿// Copyright Epic Games, Inc. All Rights Reserved.
+
+#pragma once
+
+#include "Components/ActorComponent.h"
+
+#include "CompositePlateTexturePreviewComponent.generated.h"
+
+class FLevelEditorViewportClient;
+class UCompositeLayerPlate;
+class UCompositePassColorKeyer;
+
+/**
+ * A component that is used to display the plate texture of a composite actor's plate layer in the level viewport's 'picture-in-picture' preview widget.
+ * Used by keyer passes to allow color picking from raw plate textures instead of from potentially post-processed textures on the composite actor's meshes
+ */
+UCLASS()
+class UCompositePlateTexturePreviewComponent : public UActorComponent
+{
+	GENERATED_BODY()
+
+#if WITH_EDITOR
+public:
+	
+	// UActorComponent interface
+	virtual bool GetEditorPreviewInfo(float DeltaTime, FMinimalViewInfo& ViewOut) override;
+	virtual TSharedPtr<SWidget> GetCustomEditorPreviewWidget() override;
+	// End of UActorComponent interface
+
+	/**
+	 * Displays the plate texture preview in the level editor's 'picture-in-picture'
+	 * @param InRequestingKeyer The keyer pass on the plate layer that is requesting the preview be displayed
+	 * @return true if the preview was started
+	 */
+	bool ShowPreview(TObjectPtr<UCompositePassColorKeyer> InRequestingKeyer);
+
+	/** Hides the plate texture preview */
+	void HidePreview();
+
+	/** Gets whether the plate texture preview is currently being displayed */
+	bool IsPreviewActive() const { return bShowPreview; }
+
+private:
+	/** Gets the active level editor viewport client from the level editor */
+	FLevelEditorViewportClient* GetActiveViewportClient();
+#endif
+	
+private:
+	/** Indicates if the preview should be shown */
+	bool bShowPreview = false;
+
+	/** Cached camera preview size from the editor preferences config, restored when the plate preview is hidden */
+	TOptional<float> CachedCameraPreviewSize = TOptional<float>();
+
+	/** A reference to the plate whose texture is being previewed */
+	TWeakObjectPtr<UCompositeLayerPlate> PlateToPreview = nullptr;
+	
+	/** The slate brush used to display the preview texture */
+	TSharedPtr<FSlateBrush> PreviewBrush = nullptr;
+};
diff --git a/Engine/Plugins/Compositing/Composite/Source/Composite/Private/CompositeActor.cpp b/Engine/Plugins/Compositing/Composite/Source/Composite/Private/CompositeActor.cpp
--- a/Engine/Plugins/Compositing/Composite/Source/Composite/Private/CompositeActor.cpp
+++ b/Engine/Plugins/Compositing/Composite/Source/Composite/Private/CompositeActor.cpp
@@ -24,6 +24,7 @@
 #include "Kismet/GameplayStatics.h"
 #include "Materials/MaterialInstanceDynamic.h"
 #include "SceneView.h"
+#include "Components/CompositePlateTexturePreviewComponent.h"
 #include "UObject/FortniteMainBranchObjectVersion.h"
 #include "UObject/UE5ReleaseStreamObjectVersion.h"
 #include "UObject/ObjectSaveContext.h"
@@ -273,6 +274,10 @@ ACompositeActor::ACompositeActor(const FObjectInitializer& ObjectInitializer)
 
 	ViewProjectionComponent = CreateDefaultSubobject<UCompositeViewProjectionComponent>(TEXT("CameraViewProjectionComponent"));
 	AddInstanceComponent(ViewProjectionComponent);
+
+#if WITH_EDITOR
+	PlatePreviewComponent = CreateEditorOnlyDefaultSubobject<UCompositePlateTexturePreviewComponent>(TEXT("EditorPreviewComponent"));
+#endif
 }
 
 ACompositeActor::~ACompositeActor() = default;
@@ -695,6 +700,23 @@ void ACompositeActor::SetCompositeLayers(TArray<UCompositeLayerBase*> InLayers)
 	CompositeLayers = MoveTemp(InLayers);
 }
 
+#if WITH_EDITOR
+bool ACompositeActor::ShowPlatePreview(TObjectPtr<class UCompositePassColorKeyer> InRequestingKeyer)
+{
+	return PlatePreviewComponent->ShowPreview(InRequestingKeyer);
+}
+
+void ACompositeActor::HidePlatePreview()
+{
+	PlatePreviewComponent->HidePreview();
+}
+
+bool ACompositeActor::IsPlatePreviewActive() const
+{
+	return PlatePreviewComponent->IsPreviewActive();
+}
+#endif
+
 template<>
 UCompositeSceneCapture2DComponent* ACompositeActor::FindOrCreateSceneCapture(const UCompositeLayerBase* InLayer, int32 InIndex, FName InBaseName)
 {
diff --git a/Engine/Plugins/Compositing/Composite/Source/Composite/Public/CompositeActor.h b/Engine/Plugins/Compositing/Composite/Source/Composite/Public/CompositeActor.h
--- a/Engine/Plugins/Compositing/Composite/Source/Composite/Public/CompositeActor.h
+++ b/Engine/Plugins/Compositing/Composite/Source/Composite/Public/CompositeActor.h
@@ -120,11 +120,32 @@ public:
 	UFUNCTION(BlueprintSetter)
 	UE_API void SetCompositeLayers(TArray<UCompositeLayerBase*> InLayers);
 
+#if WITH_EDITOR
+	/**
+	 * Displays the plate texture preview in the level editor's 'picture-in-picture'
+	 * @param InRequestingKeyer The keyer pass on the plate layer that is requesting the preview be displayed
+	 * @return true if the preview was started
+	 */
+	UE_API bool ShowPlatePreview(TObjectPtr<class UCompositePassColorKeyer> InRequestingKeyer);
+
+	/** Hides the plate texture preview */
+	UE_API void HidePlatePreview();
+
+	/** Gets whether the plate texture preview is currently being displayed */
+	UE_API bool IsPlatePreviewActive() const;
+#endif
+
 private:
 	/** Component responsible for continuously updating a material parameter collection with the composite camera view projection matrix. */
 	UPROPERTY(Instanced)
 	TObjectPtr<class UCompositeViewProjectionComponent> ViewProjectionComponent;
 
+#if WITH_EDITORONLY_DATA
+	/** Component used to preview a plate layer texture in the level viewport */
+	UPROPERTY(Transient)
+	TObjectPtr<class UCompositePlateTexturePreviewComponent> PlatePreviewComponent;
+#endif
+
 	/** Whether or not the composite behavior is active locally - used primarily for multi-user. */
 	UPROPERTY(NonTransactional, EditInstanceOnly, BlueprintGetter = IsActive, BlueprintSetter = SetActive, Category = "Composite", meta = (AllowPrivateAccess = true))
 	bool bIsActive;
diff --git a/Engine/Plugins/Compositing/Composite/Source/Composite/Public/Passes/CompositePassColorKeyer.h b/Engine/Plugins/Compositing/Composite/Source/Composite/Public/Passes/CompositePassColorKeyer.h
--- a/Engine/Plugins/Compositing/Composite/Source/Composite/Public/Passes/CompositePassColorKeyer.h
+++ b/Engine/Plugins/Compositing/Composite/Source/Composite/Public/Passes/CompositePassColorKeyer.h
@@ -88,6 +88,10 @@ public:
 	/** Setting that controls how long to delay before capturing a clean plate */
 	UPROPERTY(EditAnywhere, Category = "", meta = (PropertyGroup = "Keyer Settings", EditCondition = "KeyerSource == ECompositeKeyerSource::CleanPlate", EditConditionHides, ClampMin=0.0, UIMin=0.0, ClampMax=60.0, UIMax=60.0))
 	float CleanPlateCountdown = 0.0f;
+
+	/** Setting that determines how large to display the plate texture 'picture-in-picture' preview in the level editor when using the eye dropper color picker */
+	UPROPERTY(EditAnywhere, Category = "", meta = (PropertyGroup = "Keyer Settings", EditCondition = "KeyerSource == ECompositeKeyerSource::Color", EditConditionHides, ClampMin=1.0, UIMin=1.0, ClampMax=10.0, UIMax=10.0))
+	float PlateTexturePreviewSize = 5.0f;
 #endif
 	
 	/** Weight of the foreground red channel contributing to the key matte hardness. */
diff --git a/Engine/Plugins/Compositing/Composite/Source/CompositeEditor/CompositeEditor.Build.cs b/Engine/Plugins/Compositing/Composite/Source/CompositeEditor/CompositeEditor.Build.cs
--- a/Engine/Plugins/Compositing/Composite/Source/CompositeEditor/CompositeEditor.Build.cs
+++ b/Engine/Plugins/Compositing/Composite/Source/CompositeEditor/CompositeEditor.Build.cs
@@ -20,11 +20,13 @@ public class CompositeEditor : ModuleRules
 			new string[]
 			{
 				"ApplicationCore",
+				"AppFramework",
 				"ColorGradingEditor",
 				"CompositeCore",
 				"ConcertSyncClient",
 				"ContentBrowserData",
 				"CoreUObject",
+				"DetailCustomizations",
 				"EditorWidgets",
 				"EditorScriptingUtilities",
 				"Engine",
diff --git a/Engine/Plugins/Compositing/Composite/Source/CompositeEditor/Private/Customizations/CompositePassColorKeyerCustomization.cpp b/Engine/Plugins/Compositing/Composite/Source/CompositeEditor/Private/Customizations/CompositePassColorKeyerCustomization.cpp
--- a/Engine/Plugins/Compositing/Composite/Source/CompositeEditor/Private/Customizations/CompositePassColorKeyerCustomization.cpp
+++ b/Engine/Plugins/Compositing/Composite/Source/CompositeEditor/Private/Customizations/CompositePassColorKeyerCustomization.cpp
@@ -3,6 +3,7 @@
 #include "CompositePassColorKeyerCustomization.h"
 
 #include "AssetToolsModule.h"
+#include "CompositeActor.h"
 #include "CompositeEditorModule.h"
 #include "CompositeEditorStyle.h"
 #include "ContentBrowserModule.h"
@@ -14,16 +15,22 @@
 #include "FileHelpers.h"
 #include "IContentBrowserSingleton.h"
 #include "IDetailGroup.h"
+#include "IPropertyUtilities.h"
+#include "LevelEditorViewport.h"
 #include "MediaTexture.h"
 #include "NumericPropertyParams.h"
 #include "ScopedTransaction.h"
+#include "SLevelViewport.h"
 #include "AssetRegistry/AssetRegistryModule.h"
+#include "Composite/Private/Components/CompositePlateTexturePreviewComponent.h"
+#include "Customizations/ColorStructCustomization.h"
 #include "Engine/Texture2D.h"
 #include "Kismet/KismetRenderingLibrary.h"
 #include "Layers/CompositeLayerPlate.h"
 #include "Materials/Material.h"
 #include "Materials/MaterialInstanceDynamic.h"
 #include "Passes/CompositePassColorKeyer.h"
+#include "Widgets/Colors/SColorPicker.h"
 #include "Widgets/Input/SButton.h"
 #include "Widgets/Input/SVectorInputBox.h"
 #include "Widgets/Layout/SWidgetSwitcher.h"
@@ -172,8 +179,251 @@ void FCompositePassColorKeyerCustomization::CustomizeDetails(const TSharedPtr<ID
 	IDetailCustomization::CustomizeDetails(DetailBuilder);
 }
 
+/**
+ * Custom color struct customization that adds an inline eye dropper button, as well as hooking the opening and closing of the full color picker to displaying
+ * the plate texture preview in the level editor
+ */
+class FKeyerColorStructCustomization : public FColorStructCustomization
+{
+public:
+	virtual void CustomizeHeader(TSharedRef<IPropertyHandle> PropertyHandle, FDetailWidgetRow& HeaderRow, IPropertyTypeCustomizationUtils& CustomizationUtils) override
+	{
+		// Find the keyer pass being customized, and store a pointer to it
+		const TArray<TWeakObjectPtr<UObject>>& SelectedObjects = CustomizationUtils.GetPropertyUtilities()->GetSelectedObjects();
+		if (SelectedObjects.Num() > 0)
+		{
+			Keyer = Cast<UCompositePassColorKeyer>(SelectedObjects[0]);
+		}
+		
+		FColorStructCustomization::CustomizeHeader(PropertyHandle, HeaderRow, CustomizationUtils);
+	}
+
+	virtual void MakeHeaderRow(TSharedRef<IPropertyHandle>& InStructPropertyHandle, FDetailWidgetRow& Row) override
+	{
+		TSharedPtr<SWidget> ColorWidget;
+		float ContentWidth = 125.0f;
+
+		TWeakPtr<IPropertyHandle> StructWeakHandlePtr = StructPropertyHandle;
+		ColorWidget = CreateColorWidget(StructWeakHandlePtr);
+		
+		SColorPicker::FEyeDropperButtonArgs ButtonArgs;
+		ButtonArgs.OnBegin.BindSP(this, &FKeyerColorStructCustomization::OnBeginEyeDropper);
+		ButtonArgs.OnComplete.BindSP(this, &FKeyerColorStructCustomization::OnEndEyeDropper);
+		ButtonArgs.OnValueChanged.BindSP(this, &FKeyerColorStructCustomization::OnSetColorFromColorPicker);
+		
+		TSharedRef<SButton> EyeDropperButton = SColorPicker::CreateStandaloneEyeDropperButton(ButtonArgs);
+		EyeDropperButton->SetButtonStyle(&FAppStyle::Get().GetWidgetStyle<FButtonStyle>("NoBorder"));
+		
+		Row.NameContent()
+		[
+			StructPropertyHandle->CreatePropertyNameWidget()
+		]
+		.ValueContent()
+		.MinDesiredWidth(ContentWidth)
+		[
+			SNew(SHorizontalBox)
+
+			+SHorizontalBox::Slot()
+			[
+				ColorWidget.ToSharedRef()
+			]
+
+			+SHorizontalBox::Slot()
+			.AutoWidth()
+			.Padding(2.0f, 0.0f, 0.0f, 0.0f)
+			[
+				EyeDropperButton
+			]
+		];
+	}
+
+private:
+	// Copy and hide the following original FColorStructCustomization methods with slight variations that allow the plate texture preview to be displayed or hidden
+	// whenever the color picker is open or closed.
+
+	TSharedRef<SWidget> CreateColorWidget(TWeakPtr<IPropertyHandle> StructWeakHandlePtr)
+	{
+		return
+			SNew(SBox)
+			.Padding(FMargin(0,0,4.0f,0.0f))
+			.VAlign(VAlign_Center)
+			[
+				SAssignNew(ColorWidgetBackgroundBorder, SBorder)
+				.Padding(1)
+				.BorderImage(FAppStyle::Get().GetBrush("ColorPicker.RoundedSolidBackground"))
+				.BorderBackgroundColor(this, &FKeyerColorStructCustomization::GetColorWidgetBorderColor)
+				.VAlign(VAlign_Center)
+				[
+					SNew(SOverlay)
+					+ SOverlay::Slot()
+					.VAlign(VAlign_Center)
+					[
+						SAssignNew(ColorPickerParentWidget, SColorBlock)
+						.AlphaBackgroundBrush(FAppStyle::Get().GetBrush("ColorPicker.RoundedAlphaBackground"))
+						.Color(this, &FKeyerColorStructCustomization::OnGetColorForColorBlock)
+						.ShowBackgroundForAlpha(true)
+						.AlphaDisplayMode(bIgnoreAlpha ? EColorBlockAlphaDisplayMode::Ignore : EColorBlockAlphaDisplayMode::Separate)
+						.OnMouseButtonDown(this, &FKeyerColorStructCustomization::OnMouseButtonDownColorBlock)
+						.Size(FVector2D(70.0f, 20.0f))
+						.CornerRadius(FVector4(4.0f,4.0f,4.0f,4.0f))
+						.IsEnabled(this, &FKeyerColorStructCustomization::IsValueEnabled, StructWeakHandlePtr)
+					]
+					+ SOverlay::Slot()
+					.VAlign(VAlign_Center)
+					[
+						SNew(SBorder)
+						.Visibility(this, &FKeyerColorStructCustomization::GetMultipleValuesTextVisibility)
+						.BorderImage(FAppStyle::Get().GetBrush("ColorPicker.MultipleValuesBackground"))
+						.VAlign(VAlign_Center)
+						.ForegroundColor(FAppStyle::Get().GetWidgetStyle<FEditableTextBoxStyle>("NormalEditableTextBox").ForegroundColor)
+						.Padding(FMargin(12.0f, 2.0f))
+						[
+							SNew(STextBlock)
+							.Text(NSLOCTEXT("PropertyEditor", "MultipleValues", "Multiple Values"))
+							.Font(IDetailLayoutBuilder::GetDetailFont())
+						]
+					]
+				]
+			];
+	}
+
+	FReply OnMouseButtonDownColorBlock(const FGeometry& MyGeometry, const FPointerEvent& MouseEvent)
+	{
+		if (MouseEvent.GetEffectingButton() != EKeys::LeftMouseButton)
+		{
+			return FReply::Unhandled();
+		}
+	
+		bool CanShowColorPicker = true;
+		if (StructPropertyHandle.IsValid() && StructPropertyHandle->GetProperty() != nullptr)
+		{
+			CanShowColorPicker = !StructPropertyHandle->IsEditConst();
+		}
+		if (CanShowColorPicker)
+		{
+			CreateColorPicker(true /*bUseAlpha*/);
+		}
+
+		return FReply::Handled();
+	}
+	
+	void CreateColorPicker(bool bUseAlpha)
+	{
+		TransactionIndex = GEditor->BeginTransaction(FText::Format(LOCTEXT("SetColorProperty", "Edit {0}"), StructPropertyHandle->GetPropertyDisplayName()));
+
+		GatherSavedPreColorPickerColors();
+
+		FLinearColor InitialColor;
+		GetColorAsLinear(InitialColor);
+
+		const bool bRefreshOnlyOnOk = bDontUpdateWhileEditing || StructPropertyHandle->HasMetaData("DontUpdateWhileEditing");
+		const bool bOnlyRefreshOnMouseUp = StructPropertyHandle->HasMetaData("OnlyUpdateOnInteractionEnd");
+
+		FColorPickerArgs PickerArgs;
+		{
+			PickerArgs.bUseAlpha = !bIgnoreAlpha;
+			PickerArgs.bOnlyRefreshOnMouseUp = bOnlyRefreshOnMouseUp;
+			PickerArgs.bOnlyRefreshOnOk = bRefreshOnlyOnOk;
+			PickerArgs.sRGBOverride = sRGBOverride;
+			PickerArgs.DisplayGamma = TAttribute<float>::Create(TAttribute<float>::FGetter::CreateUObject(GEngine, &UEngine::GetDisplayGamma));
+			PickerArgs.OnColorCommitted = FOnLinearColorValueChanged::CreateSP(this, &FKeyerColorStructCustomization::OnSetColorFromColorPicker);
+			PickerArgs.OnColorPickerCancelled = FOnColorPickerCancelled::CreateSP(this, &FKeyerColorStructCustomization::OnColorPickerCancelled);
+			PickerArgs.OnColorPickerWindowClosed = FOnWindowClosed::CreateSP(this, &FKeyerColorStructCustomization::OnColorPickerWindowClosed);
+			PickerArgs.OnInteractivePickBegin = FSimpleDelegate::CreateSP(this, &FKeyerColorStructCustomization::OnColorPickerInteractiveBegin);
+			PickerArgs.OnInteractivePickEnd = FSimpleDelegate::CreateSP(this, &FKeyerColorStructCustomization::OnColorPickerInteractiveEnd);
+			PickerArgs.InitialColor = InitialColor;
+			PickerArgs.ParentWidget = ColorPickerParentWidget;
+			PickerArgs.OptionalOwningDetailsView = ColorPickerParentWidget;
+			FWidgetPath ParentWidgetPath;
+			if (FSlateApplication::Get().FindPathToWidget(ColorPickerParentWidget.ToSharedRef(), ParentWidgetPath))
+			{
+				PickerArgs.bOpenAsMenu = FSlateApplication::Get().FindMenuInWidgetPath(ParentWidgetPath).IsValid();
+			}
+		}
+
+		if (OpenColorPicker(PickerArgs))
+		{
+			if (TStrongObjectPtr<UCompositePassColorKeyer> PinnedKeyer = Keyer.Pin())
+			{
+				if (ACompositeActor* CompositeActor = PinnedKeyer->GetTypedOuter<ACompositeActor>())
+				{
+					CompositeActor->ShowPlatePreview(PinnedKeyer.Get());
+				}
+			}
+		}
+	}
+
+	void OnColorPickerWindowClosed(const TSharedRef<SWindow>& Window)
+	{
+		// Transact only at the end to avoid opening a lingering transaction. Reset value before transacting.
+		if (!LastPickerColorString.IsEmpty())
+		{
+			StructPropertyHandle->SetValueFromFormattedString(LastPickerColorString);
+		}
+
+		GEditor->EndTransaction();
+		TransactionIndex.Reset();
+
+		if (TStrongObjectPtr<UCompositePassColorKeyer> PinnedKeyer = Keyer.Pin())
+		{
+			if (ACompositeActor* CompositeActor = PinnedKeyer->GetTypedOuter<ACompositeActor>())
+			{
+				CompositeActor->HidePlatePreview();
+			}
+		}
+	}
+
+	/** Raised when the inline eye dropper button is clicked */
+	void OnBeginEyeDropper()
+	{
+		TransactionIndex = GEditor->BeginTransaction(FText::Format(LOCTEXT("SetColorProperty", "Edit {0}"), StructPropertyHandle->GetPropertyDisplayName()));
+		
+		if (TStrongObjectPtr<UCompositePassColorKeyer> PinnedKeyer = Keyer.Pin())
+		{
+			if (ACompositeActor* CompositeActor = PinnedKeyer->GetTypedOuter<ACompositeActor>())
+			{
+				CompositeActor->ShowPlatePreview(PinnedKeyer.Get());
+			}
+		}
+	}
+
+	/** Raised when the inline eye dropper is completed, either by a color being selected or cancelled */
+	void OnEndEyeDropper(bool bCancelled)
+	{
+		if (!bCancelled)
+		{
+			GEditor->EndTransaction();
+			TransactionIndex.Reset();
+		}
+		else
+		{
+			GEditor->CancelTransaction(0);
+			TransactionIndex.Reset();
+		}
+		
+		if (TStrongObjectPtr<UCompositePassColorKeyer> PinnedKeyer = Keyer.Pin())
+		{
+			if (ACompositeActor* CompositeActor = PinnedKeyer->GetTypedOuter<ACompositeActor>())
+			{
+				CompositeActor->HidePlatePreview();
+			}
+		}
+	}
+
+private:
+	/** Reference to the keyer pass whose keyer color property is being customized */
+	TWeakObjectPtr<UCompositePassColorKeyer> Keyer = nullptr;
+};
+
 void FCompositePassColorKeyerCustomization::CustomizeDetails(IDetailLayoutBuilder& DetailBuilder)
 {
+	DetailBuilder.RegisterInstancedCustomPropertyTypeLayout(
+		NAME_LinearColor,
+		FOnGetPropertyTypeCustomizationInstance::CreateLambda([]
+		{
+			return MakeShared<FKeyerColorStructCustomization>();
+		}));
+	
 	// All UCompositePassColorKeyer properties reside in the default category (Category = ""). If this changes, will need to update this category 
 	IDetailCategoryBuilder& DefaultCategory = DetailBuilder.EditCategoryAllowNone(NAME_None);
 
@@ -190,6 +440,10 @@ void FCompositePassColorKeyerCustomization::CustomizeDetails(IDetailLayoutBuilde
 	KeyerSettingsGroup.AddPropertyRow(DetailBuilder.GetProperty(GET_MEMBER_NAME_CHECKED(UCompositePassColorKeyer, KeyColor)));
 	KeyerSettingsGroup.AddPropertyRow(DetailBuilder.GetProperty(GET_MEMBER_NAME_CHECKED(UCompositePassColorKeyer, CleanPlate)));
 
+	PreviewSizeHandle = DetailBuilder.GetProperty(GET_MEMBER_NAME_CHECKED(UCompositePassColorKeyer, PlateTexturePreviewSize));
+	KeyerSettingsGroup.AddPropertyRow(PreviewSizeHandle.ToSharedRef());
+	PreviewSizeHandle->SetOnPropertyValueChanged(FSimpleDelegate::CreateSP(this, &FCompositePassColorKeyerCustomization::OnPreviewSizeChanged));
+	
 	CountdownHandle = DetailBuilder.GetProperty(GET_MEMBER_NAME_CHECKED(UCompositePassColorKeyer, CleanPlateCountdown));
 	KeyerSettingsGroup.AddPropertyRow(CountdownHandle.ToSharedRef());
 
@@ -342,6 +596,45 @@ void FCompositePassColorKeyerCustomization::CaptureMediaTexture()
 	}
 }
 
+void FCompositePassColorKeyerCustomization::OnPreviewSizeChanged()
+{
+	TSharedPtr<IDetailLayoutBuilder> DetailBuilder = CachedDetailBuilder.Pin();
+	if (!DetailBuilder.IsValid())
+	{
+		return;
+	}
+	
+	TArray<TWeakObjectPtr<UCompositePassColorKeyer>> Keyers = DetailBuilder->GetObjectsOfTypeBeingCustomized<UCompositePassColorKeyer>();
+	if (Keyers.Num() != 1)
+	{
+		return;
+	}
+
+	TStrongObjectPtr<UCompositePassColorKeyer> PinnedKeyer = Keyers[0].Pin();
+	if (!PinnedKeyer.IsValid())
+	{
+		return;
+	}
+
+	ACompositeActor* CompositeActor = PinnedKeyer->GetTypedOuter<ACompositeActor>();
+	if (!CompositeActor)
+	{
+		return;
+	}
+
+	if (CompositeActor->IsPlatePreviewActive())
+	{
+		float PreviewSize;
+		if (PreviewSizeHandle->GetValue(PreviewSize) == FPropertyAccess::Success)
+		{
+			if (ULevelEditorViewportSettings* ViewportSettings = GetMutableDefault<ULevelEditorViewportSettings>())
+			{
+				ViewportSettings->CameraPreviewSize = FMath::Clamp(PreviewSize, 1.0f, 10.0f);
+			}
+		}
+	}
+}
+
 TSharedRef<SWidget> FCompositePassColorKeyerCustomization::CreateRGBNumericEntryWidget()
 {
 	auto CreateNumericEntryBox = [this](const TSharedPtr<IPropertyHandle>& InPropertyHandle, FLinearColor InLabelColor)
diff --git a/Engine/Plugins/Compositing/Composite/Source/CompositeEditor/Private/Customizations/CompositePassColorKeyerCustomization.h b/Engine/Plugins/Compositing/Composite/Source/CompositeEditor/Private/Customizations/CompositePassColorKeyerCustomization.h
--- a/Engine/Plugins/Compositing/Composite/Source/CompositeEditor/Private/Customizations/CompositePassColorKeyerCustomization.h
+++ b/Engine/Plugins/Compositing/Composite/Source/CompositeEditor/Private/Customizations/CompositePassColorKeyerCustomization.h
@@ -25,6 +25,9 @@ public:
 private:
 	void CaptureMediaTexture();
 
+	/** Raised when the PlateTexturePreviewSize property changes values */
+	void OnPreviewSizeChanged();
+	
 	/** Creates the custom RGB single row widget, which contains color-coded numeric spinners for each component */
 	TSharedRef<SWidget> CreateRGBNumericEntryWidget();
 	
@@ -65,6 +68,9 @@ private:
 	/** Property handle for UCompositePassColorKeyer.KeyerSource property */
 	TSharedPtr<IPropertyHandle> KeyerSourceHandle;
 
+	/** Property handle for UCompositePassColorKeyer.PlateTexturePreviewSize property */
+	TSharedPtr<IPropertyHandle> PreviewSizeHandle;
+	
 	/** Property handle for UCompositePassKeyer.CleanPlateCountdown property */
 	TSharedPtr<IPropertyHandle> CountdownHandle;
 
diff --git a/Engine/Source/Runtime/AppFramework/Private/Widgets/Colors/SColorPicker.cpp b/Engine/Source/Runtime/AppFramework/Private/Widgets/Colors/SColorPicker.cpp
--- a/Engine/Source/Runtime/AppFramework/Private/Widgets/Colors/SColorPicker.cpp
+++ b/Engine/Source/Runtime/AppFramework/Private/Widgets/Colors/SColorPicker.cpp
@@ -596,6 +596,15 @@ bool SColorPicker::SetNewTargetColorRGB( const FLinearColor& NewValue, bool bFor
 }
 
 
+TSharedRef<SButton> SColorPicker::CreateStandaloneEyeDropperButton(const FEyeDropperButtonArgs& InArgs)
+{
+	return SNew(SEyeDropperButton)
+		.OnValueChanged(InArgs.OnValueChanged)
+		.OnBegin(InArgs.OnBegin)
+		.OnComplete(InArgs.OnComplete)
+		.DisplayGamma(InArgs.DisplayGamma);
+}
+
 bool SColorPicker::ApplyNewTargetColor( bool bForceUpdate /*= false*/ )
 {
 	bool bUpdated = false;
diff --git a/Engine/Source/Runtime/AppFramework/Public/Widgets/Colors/SColorPicker.h b/Engine/Source/Runtime/AppFramework/Public/Widgets/Colors/SColorPicker.h
--- a/Engine/Source/Runtime/AppFramework/Public/Widgets/Colors/SColorPicker.h
+++ b/Engine/Source/Runtime/AppFramework/Public/Widgets/Colors/SColorPicker.h
@@ -216,6 +216,25 @@ public:
 	DECLARE_DELEGATE(FOnColorPickerDestructionOverride);
 	static APPFRAMEWORK_API FOnColorPickerDestructionOverride OnColorPickerDestroyOverride;
 
+	/** Eye dropper button creation args */
+	struct FEyeDropperButtonArgs
+	{
+		/** Invoked when a new value is selected by the dropper */
+		TDelegate<void(FLinearColor)> OnValueChanged;
+
+		/** Invoked when the dropper goes from inactive to active */
+		FSimpleDelegate OnBegin;
+
+		/** Invoked when the dropper goes from active to inactive */
+		TDelegate<void(bool)> OnComplete;
+
+		/** Sets the display Gamma setting - used to correct colors sampled from the screen */
+		TAttribute<float> DisplayGamma = TAttribute<float>();
+	};
+
+	/** Creates a standalone eye dropper button that functions identically to the button within the color picker */
+	static APPFRAMEWORK_API TSharedRef<SButton> CreateStandaloneEyeDropperButton(const FEyeDropperButtonArgs& InArgs);
+
 protected:
 
 	/** Backup all the colors that are being modified */
