Composure: Fixed support for multiple keyer passes in a sequence.

* We were previously incorrectly discarding the original alpha value.

#Synced-CL 50287290

diff --git a/Engine/Plugins/Compositing/Composite/Shaders/Private/CompositeColorKeyer.usf b/Engine/Plugins/Compositing/Composite/Shaders/Private/CompositeColorKeyer.usf
--- a/Engine/Plugins/Compositing/Composite/Shaders/Private/CompositeColorKeyer.usf
+++ b/Engine/Plugins/Compositing/Composite/Shaders/Private/CompositeColorKeyer.usf
@@ -98,15 +98,15 @@ void MainCS(
 	BRANCH
 	if(ScreenType == 0)
 	{
-		Result = RedScreenKeyer(Result.rgb, Background, ChannelWeights[1], ChannelWeights[2]);
+		Result = RedScreenKeyer(Result, Background, ChannelWeights[1], ChannelWeights[2]);
 	}
 	else if(ScreenType == 1)
 	{
-		Result = GreenScreenKeyer(Result.rgb, Background, ChannelWeights[0], ChannelWeights[2]);
+		Result = GreenScreenKeyer(Result, Background, ChannelWeights[0], ChannelWeights[2]);
 	}
 	else if(ScreenType == 2)
 	{
-		Result = BlueScreenKeyer(Result.rgb, Background, ChannelWeights[0], ChannelWeights[1]);
+		Result = BlueScreenKeyer(Result, Background, ChannelWeights[0], ChannelWeights[1]);
 	}
 
 	BRANCH
diff --git a/Engine/Plugins/Compositing/Composite/Shaders/Private/CompositeColorKeyerCommon.ush b/Engine/Plugins/Compositing/Composite/Shaders/Private/CompositeColorKeyerCommon.ush
--- a/Engine/Plugins/Compositing/Composite/Shaders/Private/CompositeColorKeyerCommon.ush
+++ b/Engine/Plugins/Compositing/Composite/Shaders/Private/CompositeColorKeyerCommon.ush
@@ -171,7 +171,7 @@ float Vignette(float2 ScreenPos, float AspectRatio, float Intensity)
 }
 
 float4 GreenScreenKeyer(
-	float3 Foreground,
+	float4 Foreground,
 	float3 Background,
 	float  RedWeight,
 	float  BlueWeight)
@@ -197,7 +197,7 @@ float4 GreenScreenKeyer(
 	float3 MaskedB = DespillB * GreenRatio;
 
 	// Handle screen brighter than sample color
-	float3 Delta = Foreground - Background;
+	float3 Delta = Foreground.rgb - Background;
 	float Mg     = Delta.g - (RedWeight * Delta.r + BlueWeight * Delta.b);
 	MaskedB.g   += saturate(Mg * InvBg * wB);
 
@@ -211,6 +211,8 @@ float4 GreenScreenKeyer(
 
 	// Hack to make sure the bright areas of the plate are solid
 	Alpha = max(Alpha, Fill.r);
+	// Preserve original alpha
+	Alpha *= Foreground.a;
 	// Clamp alpha
 	Alpha = saturate(Alpha);
 
@@ -230,9 +232,19 @@ float3 SwizzleGB(float3 InColor)
 	return float3(InColor.r, InColor.b, InColor.g);
 }
 
-float4 RedScreenKeyer(float3 Foreground, float3 Background, float GreenWeight, float BlueWeight)
+float4 SwizzleRG(float4 InColor)
 {
-	float3 SwizzledForeground = SwizzleRG(Foreground);
+	return float4(InColor.g, InColor.r, InColor.b, InColor.a);
+}
+
+float4 SwizzleGB(float4 InColor)
+{
+	return float4(InColor.r, InColor.b, InColor.g, InColor.a);
+}
+
+float4 RedScreenKeyer(float4 Foreground, float3 Background, float GreenWeight, float BlueWeight)
+{
+	float4 SwizzledForeground = SwizzleRG(Foreground);
 	float3 SizzledBackground = SwizzleRG(Background);
 	float4 Output = GreenScreenKeyer(SwizzledForeground, SizzledBackground, GreenWeight, BlueWeight);
 	Output.rgb = SwizzleRG(Output.rgb);
@@ -240,9 +252,9 @@ float4 RedScreenKeyer(float3 Foreground, float3 Background, float GreenWeight, f
 	return Output;
 }
 
-float4 BlueScreenKeyer(float3 Foreground, float3 Background, float RedWeight, float GreenWeight)
+float4 BlueScreenKeyer(float4 Foreground, float3 Background, float RedWeight, float GreenWeight)
 {
-	float3 SwizzledForeground = SwizzleGB(Foreground);
+	float4 SwizzledForeground = SwizzleGB(Foreground);
 	float3 SizzledBackground = SwizzleGB(Background);
 	float4 Output = GreenScreenKeyer(SwizzledForeground, SizzledBackground, RedWeight, GreenWeight);
 	Output.rgb = SwizzleGB(Output.rgb);
