Composure: Fixed double-DOF artifacts with translucent keyed footage applied on composite meshes.

* We automatically update required rendering cvars when AfterDOF translucent materials are in use on layer plate composite meshes.
* Relies on the recent engine translucency holdout fix CL 50271663.
* Refactored console variable handling into a manager class for convenience.

#Synced-CL 50352176

diff --git a/Engine/Plugins/Compositing/Composite/Source/Composite/Private/CompositeCVarOverrideManager.cpp b/Engine/Plugins/Compositing/Composite/Source/Composite/Private/CompositeCVarOverrideManager.cpp
new file mode 100644
--- /dev/null
+++ b/Engine/Plugins/Compositing/Composite/Source/Composite/Private/CompositeCVarOverrideManager.cpp
@@ -0,0 +1,116 @@
+// Copyright Epic Games, Inc. All Rights Reserved.
+
+#include "CompositeCVarOverrideManager.h"
+
+#include "CompositeModule.h"
+#include "HAL/IConsoleManager.h"
+
+static bool GCompositeAllowCVarOverrides = true;
+static FAutoConsoleVariableRef CVarCompositeAllowCVarOverrides(
+	TEXT("Composite.AllowCVarOverrides"),
+	GCompositeAllowCVarOverrides,
+	TEXT("Read-only variable to disallow the composite cvar manager from overriding other engine cvars.\n")
+	TEXT("0: Disallow plugin overrides\n")
+	TEXT("1: Allow plugin overrides"),
+	ECVF_ReadOnly
+);
+
+FCompositeCVarOverrideManager& FCompositeCVarOverrideManager::Get()
+{
+	static FCompositeCVarOverrideManager Instance;
+	
+	return Instance;
+}
+
+template <typename T>
+void FCompositeCVarOverrideManager::Override(const TCHAR* CVarName, T NewValue)
+{
+	if (!GCompositeAllowCVarOverrides)
+	{
+		return;
+	}
+
+	IConsoleVariable* CVar = IConsoleManager::Get().FindConsoleVariable(CVarName);
+
+	if (!CVar)
+	{
+		UE_LOG(LogComposite, Warning, TEXT("Could not find CVar '%s'"), CVarName);
+		return;
+	}
+
+	if (!OverrideCache.Contains(CVarName))
+	{
+		FCachedCVarData Data;
+		Data.CVarPtr = CVar;
+		Data.OriginalValue = CVar->GetString(); // String is valid for all types
+
+		OverrideCache.Add(CVarName, Data);
+	}
+
+	CVar->Set(NewValue, ECVF_SetByCode);
+
+	UE_LOG(LogComposite, Display, TEXT("Applying '%s=%s' override."), CVarName, *CVar->GetString());
+}
+
+template void FCompositeCVarOverrideManager::Override<float>(const TCHAR* CVarName, float NewValue);
+template void FCompositeCVarOverrideManager::Override<int32>(const TCHAR* CVarName, int32 NewValue);
+template void FCompositeCVarOverrideManager::Override<bool>(const TCHAR* CVarName, bool NewValue);
+
+void FCompositeCVarOverrideManager::Restore(const TCHAR* CVarName)
+{
+	FCachedCVarData Data;
+	if (OverrideCache.RemoveAndCopyValue(CVarName, Data))
+	{
+		if (Data.CVarPtr)
+		{
+			Data.CVarPtr->Set(*Data.OriginalValue, ECVF_SetByCode);
+
+			UE_LOG(LogComposite, Display, TEXT("Restored '%s=%s'"), CVarName, *Data.CVarPtr->GetString());
+		}
+	}
+}
+
+void FCompositeCVarOverrideManager::RestoreAll()
+{
+	for (auto It = OverrideCache.CreateIterator(); It; ++It)
+	{
+		FCachedCVarData& Data = It.Value();
+		if (Data.CVarPtr)
+		{
+			Data.CVarPtr->Set(*Data.OriginalValue, ECVF_SetByCode);
+		}
+	}
+
+	OverrideCache.Empty();
+}
+
+void FCompositeCVarOverrideManager::PrintActiveOverrides()
+{
+	if (OverrideCache.IsEmpty())
+	{
+		UE_LOG(LogComposite, Display, TEXT("No active overrides."));
+		return;
+	}
+
+	UE_LOG(LogComposite, Display, TEXT("=== Active Overrides (%d) ==="), OverrideCache.Num());
+
+	for (auto It = OverrideCache.CreateConstIterator(); It; ++It)
+	{
+		const FString& CVarName = It.Key();
+		const FCachedCVarData& Data = It.Value();
+
+		FString CurrentValue = TEXT("INVALID");
+		if (Data.CVarPtr)
+		{
+			// Get the live value from the engine
+			CurrentValue = Data.CVarPtr->GetString();
+		}
+
+		UE_LOG(LogComposite, Log, TEXT("  %s : Original[%s] -> Current[%s]"),
+			*CVarName,
+			*Data.OriginalValue,
+			*CurrentValue);
+	}
+
+	UE_LOG(LogComposite, Display, TEXT("=================================================="));
+}
diff --git a/Engine/Plugins/Compositing/Composite/Source/Composite/Private/CompositeCVarOverrideManager.h b/Engine/Plugins/Compositing/Composite/Source/Composite/Private/CompositeCVarOverrideManager.h
new file mode 100644
--- /dev/null
+++ b/Engine/Plugins/Compositing/Composite/Source/Composite/Private/CompositeCVarOverrideManager.h
@@ -0,0 +1,46 @@
+// Copyright Epic Games, Inc. All Rights Reserved.
+
+#pragma once
+
+#include "CoreTypes.h"
+#include "Containers/UnrealString.h"
+#include "Containers/SortedMap.h"
+
+class IConsoleVariable;
+
+/** Console variable override manager. */
+class FCompositeCVarOverrideManager
+{
+public:
+	/** Singleton getter. */
+	static FCompositeCVarOverrideManager& Get();
+
+	FCompositeCVarOverrideManager(const FCompositeCVarOverrideManager&) = delete;
+	FCompositeCVarOverrideManager& operator=(const FCompositeCVarOverrideManager&) = delete;
+
+	/** Overrides a console variable with a new value. */
+	template <typename T>
+	void Override(const TCHAR* CVarName, T NewValue);
+
+	/** Restore a specific console variable to its original cached value. */
+	void Restore(const TCHAR* CVarName);
+
+	/** Restore all overrides to their original values. */
+	void RestoreAll();
+
+	/** Debug function to print all active overrides. */
+	void PrintActiveOverrides();
+
+private:
+	FCompositeCVarOverrideManager() = default;
+	~FCompositeCVarOverrideManager() = default;
+
+	struct FCachedCVarData
+	{
+		IConsoleVariable* CVarPtr = nullptr;
+		FString OriginalValue;
+	};
+
+	/** Cached original cvar data */
+	TSortedMap<FString, FCachedCVarData> OverrideCache;
+};
diff --git a/Engine/Plugins/Compositing/Composite/Source/Composite/Private/Layers/CompositeLayerPlate.cpp b/Engine/Plugins/Compositing/Composite/Source/Composite/Private/Layers/CompositeLayerPlate.cpp
--- a/Engine/Plugins/Compositing/Composite/Source/Composite/Private/Layers/CompositeLayerPlate.cpp
+++ b/Engine/Plugins/Compositing/Composite/Source/Composite/Private/Layers/CompositeLayerPlate.cpp
@@ -7,6 +7,7 @@
 #include "CompositeAssetUserData.h"
 #include "CompositeCoreSubsystem.h"
 #include "CompositeCoreSettings.h"
+#include "CompositeCVarOverrideManager.h"
 #include "CompositeRenderTargetPool.h"
 #include "Engine/Texture2D.h"
 #include "MediaTexture.h"
@@ -659,6 +660,51 @@ void UCompositeLayerPlate::OnPostSaveWorld(UWorld* InWorld, FObjectPostSaveConte
 }
 #endif
 
+void UCompositeLayerPlate::OverrideTranslucencyConsoleVariables(bool bOverride) const
+{
+	static TSet<const UCompositeLayerPlate*> ActiveOverrideLayer;
+
+	const TCHAR* CVarAutoBeforeDOFName = TEXT("r.Translucency.AutoBeforeDOF");
+	const TCHAR* CVarHoldoutLocationName = TEXT("r.Translucency.Holdout.Location");
+
+	if (bOverride)
+	{
+		if (ActiveOverrideLayer.IsEmpty())
+		{
+			FCompositeCVarOverrideManager::Get().Override(CVarAutoBeforeDOFName, -1.0f);
+			FCompositeCVarOverrideManager::Get().Override(CVarHoldoutLocationName, 1);
+		}
+
+		ActiveOverrideLayer.Add(this);
+	}
+	else
+	{
+		const int32 NumRemoved = ActiveOverrideLayer.Remove(this);
+
+		if (NumRemoved > 0 && ActiveOverrideLayer.IsEmpty())
+		{
+			FCompositeCVarOverrideManager::Get().Restore(CVarAutoBeforeDOFName);
+			FCompositeCVarOverrideManager::Get().Restore(CVarHoldoutLocationName);
+		}
+	}
+}
+
+bool UCompositeLayerPlate::IsUsingAfterDOFTranslucency(UMaterialInterface* MaterialInterface) const
+{
+	if (IsTranslucentBlendMode(*MaterialInterface))
+	{
+		if (FMaterialResource* Resource = MaterialInterface->GetMaterialResource(GMaxRHIShaderPlatform))
+		{
+			if (Resource->IsTranslucencyAfterDOFEnabled())
+			{
+				return true;
+			}
+		}
+	}
+
+	return false;
+}
+
 void UCompositeLayerPlate::UpdateCompositeMeshes(bool bRegisterOnly) const
 {
 	for (const TSoftObjectPtr<AActor>& MeshActor : CompositeMeshes)
@@ -708,6 +754,12 @@ void UCompositeLayerPlate::UpdatePrimitiveComponent(UPrimitiveComponent& InPrimi
 		TArray<FGuid> ParameterIds;
 		MaterialInterface->GetAllTextureParameterInfo(ParameterInfo, ParameterIds);
 
+		if (IsUsingAfterDOFTranslucency(MaterialInterface))
+		{
+			constexpr bool bOverride = true;
+			OverrideTranslucencyConsoleVariables(bOverride);
+		}
+
 		// Avoid conversion to MID if parameter doesn't exist
 		const bool bHasTextureParameter = ParameterInfo.ContainsByPredicate([](const FMaterialParameterInfo& ParamInfo)
 			{
@@ -738,15 +790,15 @@ void UCompositeLayerPlate::UpdatePrimitiveComponent(UPrimitiveComponent& InPrimi
 
 void UCompositeLayerPlate::PropagateStateChange(bool bApply, const UWorld* SourceWorld) const
 {
-	if (CompositeMeshes.IsEmpty())
-	{
-		return;
-	}
-
 	UCompositeCoreSubsystem* Subsystem = UWorld::GetSubsystem<UCompositeCoreSubsystem>(SourceWorld ? SourceWorld : GetWorld());
 
 	if (bApply && IsRendering())
 	{
+		if (CompositeMeshes.IsEmpty())
+		{
+			return;
+		}
+
 		constexpr bool bRegisterOnly = true;
 		UpdateCompositeMeshes(bRegisterOnly);
 
@@ -755,15 +807,48 @@ void UCompositeLayerPlate::PropagateStateChange(bool bApply, const UWorld* Sourc
 			// Rely on the scene view extension to mark primivites as holdout, and render them separately.
 			Subsystem->RegisterPrimitives(GetPrimitives());
 		}
+
+		// Conditionally override translucency console variables
+		for (const TSoftObjectPtr<AActor>& MeshActor : CompositeMeshes)
+		{
+			if (!MeshActor.IsValid())
+			{
+				continue;
+			}
+
+			for (UActorComponent* Component : MeshActor->GetComponents())
+			{
+				if (UPrimitiveComponent* PrimitiveComponent = Cast<UPrimitiveComponent>(Component))
+				{
+					for (int32 ElementIndex = 0; ElementIndex < PrimitiveComponent->GetNumMaterials(); ++ElementIndex)
+					{
+						if (UMaterialInterface* MaterialInterface = PrimitiveComponent->GetMaterial(ElementIndex))
+						{
+							if (IsUsingAfterDOFTranslucency(MaterialInterface))
+							{
+								constexpr bool bOverride = true;
+								OverrideTranslucencyConsoleVariables(bOverride);
+							}
+						}
+					}
+				}
+			}
+		}
 	}
 	else
 	{
-		UE::Composite::Private::RemoveCompositeAllAssetUserData(CompositeMeshes);
-
-		if (IsValid(Subsystem))
+		if (!CompositeMeshes.IsEmpty())
 		{
-			Subsystem->UnregisterPrimitives(GetPrimitives());
+			UE::Composite::Private::RemoveCompositeAllAssetUserData(CompositeMeshes);
+
+			if (IsValid(Subsystem))
+			{
+				Subsystem->UnregisterPrimitives(GetPrimitives());
+			}
 		}
+
+		constexpr bool bOverride = false;
+		OverrideTranslucencyConsoleVariables(bOverride);
 	}
 }
 
diff --git a/Engine/Plugins/Compositing/Composite/Source/Composite/Private/Layers/CompositeLayerShadowReflection.cpp b/Engine/Plugins/Compositing/Composite/Source/Composite/Private/Layers/CompositeLayerShadowReflection.cpp
--- a/Engine/Plugins/Compositing/Composite/Source/Composite/Private/Layers/CompositeLayerShadowReflection.cpp
+++ b/Engine/Plugins/Compositing/Composite/Source/Composite/Private/Layers/CompositeLayerShadowReflection.cpp
@@ -5,6 +5,7 @@
 #include "CompositeActor.h"
 #include "Components/PrimitiveComponent.h"
 #include "Components/CompositeShadowReflectionCatcherComponent.h"
+#include "CompositeCVarOverrideManager.h"
 #include "CompositeModule.h"
 #include "CompositeRenderTargetPool.h"
 #include "Containers/StaticArray.h"
@@ -262,14 +263,9 @@ void UCompositeLayerShadowReflection::UpdatePrimitiveVisibilityState(
 void UCompositeLayerShadowReflection::SetHiddenPrimitiveState(bool bInHiddenInSceneCapture) const
 {
 	//* Note: We aim to expose megalights screen traces as an PPV setting, in which case the cvar could be removed. */
-
 	const TCHAR* MegaLightsSceenTracesName = TEXT("r.MegaLights.ScreenTraces");
 	const TCHAR* LumenShortRangeAOHWRTName = TEXT("r.Lumen.ScreenProbeGather.ShortRangeAO.HardwareRayTracing");
 
-	IConsoleVariable* MegaLightsSceenTracesCVar = IConsoleManager::Get().FindConsoleVariable(MegaLightsSceenTracesName);
-	IConsoleVariable* LumenShortRangeAOHWRTCVar = IConsoleManager::Get().FindConsoleVariable(LumenShortRangeAOHWRTName);
-
-	static TSortedMap<FString, int32> CachedConsoleVariables;
 	static TSet<const UCompositeLayerShadowReflection*> ActiveHiddenStateLayers;
 
 	if (bInHiddenInSceneCapture)
@@ -281,21 +277,8 @@ void UCompositeLayerShadowReflection::SetHiddenPrimitiveState(bool bInHiddenInSc
 
 		if (ActiveHiddenStateLayers.IsEmpty())
 		{
-			if (MegaLightsSceenTracesCVar != nullptr)
-			{
-				CachedConsoleVariables.FindOrAdd(MegaLightsSceenTracesName) = MegaLightsSceenTracesCVar->GetInt();
-				MegaLightsSceenTracesCVar->Set(0);
-
-				UE_LOG(LogComposite, Display, TEXT("Applying '%s = %d' override for shadow/reflection layer."), MegaLightsSceenTracesName, MegaLightsSceenTracesCVar->GetInt());
-			}
-
-			if (LumenShortRangeAOHWRTCVar != nullptr)
-			{
-				CachedConsoleVariables.FindOrAdd(LumenShortRangeAOHWRTName) = LumenShortRangeAOHWRTCVar->GetInt();
-				LumenShortRangeAOHWRTCVar->Set(1);
-
-				UE_LOG(LogComposite, Display, TEXT("Applying '%s = %d' override for shadow/reflection layer."), LumenShortRangeAOHWRTName, LumenShortRangeAOHWRTCVar->GetInt());
-			}
+			FCompositeCVarOverrideManager::Get().Override(MegaLightsSceenTracesName, 0);
+			FCompositeCVarOverrideManager::Get().Override(LumenShortRangeAOHWRTName, 1);
 		}
 
 		ActiveHiddenStateLayers.Add(this);
@@ -311,20 +294,8 @@ void UCompositeLayerShadowReflection::SetHiddenPrimitiveState(bool bInHiddenInSc
 
 		if (NumRemoved > 0 && ActiveHiddenStateLayers.IsEmpty())
 		{
-			int32 PreviousValue = -1;
-			if (MegaLightsSceenTracesCVar && CachedConsoleVariables.RemoveAndCopyValue(MegaLightsSceenTracesName, PreviousValue))
-			{
-				MegaLightsSceenTracesCVar->Set(PreviousValue);
-
-				UE_LOG(LogComposite, Display, TEXT("Removing shadow/reflection layer override, revert to '%s = %d'."), MegaLightsSceenTracesName, MegaLightsSceenTracesCVar->GetInt());
-			}
-
-			if (LumenShortRangeAOHWRTCVar && CachedConsoleVariables.RemoveAndCopyValue(LumenShortRangeAOHWRTName, PreviousValue))
-			{
-				LumenShortRangeAOHWRTCVar->Set(PreviousValue);
-
-				UE_LOG(LogComposite, Display, TEXT("Removing shadow/reflection layer override, revert to '%s = %d'."), LumenShortRangeAOHWRTName, LumenShortRangeAOHWRTCVar->GetInt());
-			}
+			FCompositeCVarOverrideManager::Get().Restore(MegaLightsSceenTracesName);
+			FCompositeCVarOverrideManager::Get().Restore(LumenShortRangeAOHWRTName);
 		}
 	}
 }
diff --git a/Engine/Plugins/Compositing/Composite/Source/Composite/Public/Layers/CompositeLayerPlate.h b/Engine/Plugins/Compositing/Composite/Source/Composite/Public/Layers/CompositeLayerPlate.h
--- a/Engine/Plugins/Compositing/Composite/Source/Composite/Public/Layers/CompositeLayerPlate.h
+++ b/Engine/Plugins/Compositing/Composite/Source/Composite/Public/Layers/CompositeLayerPlate.h
@@ -9,6 +9,7 @@
 
 #define UE_API COMPOSITE_API
 
+class UMaterialInterface;
 class UPrimitiveComponent;
 class UTexture;
 class FObjectPreSaveContext;
@@ -207,6 +208,15 @@ private:
 	void OnPostSaveWorld(UWorld* InWorld, FObjectPostSaveContext ObjectSaveContext);
 #endif
 
+	/**
+	 * Override the required translucency console variables.
+	 * These overrides are needed to avoid double depth-of-field on keyed footaged + composite meshes.
+	*/
+	void OverrideTranslucencyConsoleVariables(bool bOverride) const;
+
+	/** Convenience function to check if the material is AfterDOF translucent. */
+	bool IsUsingAfterDOFTranslucency(UMaterialInterface* MaterialInterface) const;
+
 	/** Processed media texture render target.*/
 	UPROPERTY(Transient, NonTransactional)
 	mutable TObjectPtr<UTextureRenderTarget2D> MediaRenderTarget;
