Color Grading: Fixed issue where Composure panel would open when performing unrelated actions in the color grading drawer

#Synced-CL 50312710

#Synced-CL 50312710

diff --git a/Engine/Plugins/Compositing/Composite/Source/CompositeEditor/Private/UI/SCompositeEditorPanel.cpp b/Engine/Plugins/Compositing/Composite/Source/CompositeEditor/Private/UI/SCompositeEditorPanel.cpp
--- a/Engine/Plugins/Compositing/Composite/Source/CompositeEditor/Private/UI/SCompositeEditorPanel.cpp
+++ b/Engine/Plugins/Compositing/Composite/Source/CompositeEditor/Private/UI/SCompositeEditorPanel.cpp
@@ -431,11 +431,17 @@ void SCompositeEditorPanel::UpdateActivePanelSelection(const TArray<UObject*>& I
 	FLevelEditorModule& LevelEditorModule = FModuleManager::LoadModuleChecked<FLevelEditorModule>("LevelEditor");
 	if (TSharedPtr<FTabManager> LevelEditorTabManager = LevelEditorModule.GetLevelEditorTabManager())
 	{
-		if (TSharedPtr<SDockTab> Tab = LevelEditorTabManager->TryInvokeTab(CompositingEditorPanel::CompositeEditorTabName))
+		if (TSharedPtr<SDockTab> Tab = LevelEditorTabManager->FindExistingLiveTab(CompositingEditorPanel::CompositeEditorTabName))
 		{
 			if (TSharedPtr<SCompositeEditorPanel> CompositePanel = StaticCastSharedPtr<SCompositeEditorPanel>(Tab->GetContent().ToSharedPtr()))
 			{
 				CompositePanel->SelectObjects(InSelection);
+
+				// If the composure panel has a valid selection after selecting the objects, draw attention to the panel so it gets put into view
+				if (CompositePanel->GetSelectedObjects().Num() > 0)
+				{
+					LevelEditorTabManager->DrawAttention(Tab.ToSharedRef());
+				}
 			}
 		}
 	}
@@ -446,7 +452,7 @@ TArray<UObject*> SCompositeEditorPanel::GetActivePanelSelection()
 	FLevelEditorModule& LevelEditorModule = FModuleManager::LoadModuleChecked<FLevelEditorModule>("LevelEditor");
 	if (TSharedPtr<FTabManager> LevelEditorTabManager = LevelEditorModule.GetLevelEditorTabManager())
 	{
-		if (TSharedPtr<SDockTab> Tab = LevelEditorTabManager->TryInvokeTab(CompositingEditorPanel::CompositeEditorTabName))
+		if (TSharedPtr<SDockTab> Tab = LevelEditorTabManager->FindExistingLiveTab(CompositingEditorPanel::CompositeEditorTabName))
 		{
 			if (TSharedPtr<SCompositeEditorPanel> CompositePanel = StaticCastSharedPtr<SCompositeEditorPanel>(Tab->GetContent().ToSharedPtr()))
 			{
diff --git a/Engine/Plugins/Editor/ObjectMixer/ObjectMixer/Source/ObjectMixer/Private/Views/List/Modes/ObjectMixerOutlinerMode.cpp b/Engine/Plugins/Editor/ObjectMixer/ObjectMixer/Source/ObjectMixer/Private/Views/List/Modes/ObjectMixerOutlinerMode.cpp
--- a/Engine/Plugins/Editor/ObjectMixer/ObjectMixer/Source/ObjectMixer/Private/Views/List/Modes/ObjectMixerOutlinerMode.cpp
+++ b/Engine/Plugins/Editor/ObjectMixer/ObjectMixer/Source/ObjectMixer/Private/Views/List/Modes/ObjectMixerOutlinerMode.cpp
@@ -2614,6 +2614,26 @@ void FObjectMixerOutlinerMode::SynchronizeComponentAndActorSelection()
 		}
 	}
 
+	Selection.ForEachItem<FObjectMixerEditorListRowUObject>([this, &SelectedActors](FObjectMixerEditorListRowUObject& Item)
+	{
+		if (UObject* Object = Item.ObjectSoftPtr.Get())
+		{
+			if (AActor* Actor = Object->GetTypedOuter<AActor>())
+			{
+				if (!SelectedActors.Contains(Actor))
+				{
+					// If there is a subobject currently selected whose parent is not in the selection being synced, deselect it
+					SceneOutliner->SetItemSelection(Item.AsShared(), false);
+				}
+				else
+				{
+					// Otherwise, don't reselect the parent actor as one of its subobjects is already selected, to avoid multi-select issues
+					SelectedActors.Remove(Actor);
+				}
+			}
+		}
+	});
+	
 	for (AActor* Actor : SelectedActors)
 	{
 		if (FSceneOutlinerTreeItemPtr ActorItem = SceneOutliner->GetTreeItem(Actor))
