Composite: Allow composite mesh use (on the plate layer) without holdouts.

For certain use cases, post-processing of the media pixels matters less than direct integration onto 3D meshes in the scene. Users in 5.7 were disabling the plate layer to avoid holdouts in favor of native renderer materials & lighting. Disabling the entire plate layer is undesirable as it removes active media passes or material updates, etc.

With this change, users can do "less" compositing, but retain the usual projection functionality. Usefulness will increase with an accurate tone curve inversion media pass (ACES 2.0), and eventual motion vector generation or DLSS use to avoid ghosting.

#Synced-CL 50473050

diff --git a/Engine/Plugins/Compositing/Composite/Source/Composite/Private/Layers/CompositeLayerPlate.cpp b/Engine/Plugins/Compositing/Composite/Source/Composite/Private/Layers/CompositeLayerPlate.cpp
--- a/Engine/Plugins/Compositing/Composite/Source/Composite/Private/Layers/CompositeLayerPlate.cpp
+++ b/Engine/Plugins/Compositing/Composite/Source/Composite/Private/Layers/CompositeLayerPlate.cpp
@@ -79,6 +79,7 @@ UCompositeLayerPlate::UCompositeLayerPlate(const FObjectInitializer& ObjectIniti
 	static ConstructorHelpers::FObjectFinder<UTexture2D> DefaultMediaTexture = TEXT("/Composite/Textures/T_Composite_SMPTE_Color_Bars_16x9.T_Composite_SMPTE_Color_Bars_16x9");
 	
 	Texture = Cast<UTexture>(DefaultMediaTexture.Object);
+	bIsHoldoutEnabled = true;
 	PlateMode = ECompositePlateMode::CompositeMesh;
 	Operation = ECompositeCoreMergeOp::Over;
 
@@ -112,46 +113,39 @@ void UCompositeLayerPlate::Tick(float DeltaTime)
 		return;
 	}
 
-	UCompositeCoreSubsystem* Subsystem = UWorld::GetSubsystem<UCompositeCoreSubsystem>(GetWorld());
-	
-	if (IsValid(Subsystem))
-	{
-		// TODO: We will need to separate the holdout feature from the custom render pass in the scene view extension, so that one can be used without the other.
-		Subsystem->RegisterPrimitives(GetPrimitives());
-	}
+	// Update holdout on tick in case composite meshes dynamically recreate their primitives
+	TryRegisterCompositePrimitives();
 
-	{
-		// Constantly check if passes have changed to manage render targets & texture bindings on composite meshes.
-		const int32 NumValidMediaPasses = GetValidPassesNum(MediaPasses);
-		const int32 NumValidScenePasses = GetValidPassesNum(ScenePasses);
-		bool bPassesHaveChanged = false;
+	// Constantly check if passes have changed to manage render targets & texture bindings on composite meshes.
+	const int32 NumValidMediaPasses = GetValidPassesNum(MediaPasses);
+	const int32 NumValidScenePasses = GetValidPassesNum(ScenePasses);
+	bool bPassesHaveChanged = false;
 
-		if (NumValidMediaPasses != CachedValidMediaPasses)
+	if (NumValidMediaPasses != CachedValidMediaPasses)
+	{
+		if (NumValidMediaPasses == 0)
 		{
-			if (NumValidMediaPasses == 0)
-			{
-				FCompositeRenderTargetPool::Get().ReleaseTarget(MediaRenderTarget);
-			}
-
-			CachedValidMediaPasses = NumValidMediaPasses;
-			bPassesHaveChanged = true;
+			FCompositeRenderTargetPool::Get().ReleaseTarget(MediaRenderTarget);
 		}
 
-		if (NumValidScenePasses != CachedValidScenePasses)
-		{
-			if (NumValidScenePasses == 0)
-			{
-				FCompositeRenderTargetPool::Get().ReleaseTarget(SceneRenderTarget);
-			}
-
-			CachedValidScenePasses = NumValidScenePasses;
-			bPassesHaveChanged = true;
-		}
+		CachedValidMediaPasses = NumValidMediaPasses;
+		bPassesHaveChanged = true;
+	}
 
-		if (bPassesHaveChanged)
+	if (NumValidScenePasses != CachedValidScenePasses)
+	{
+		if (NumValidScenePasses == 0)
 		{
-			UpdateCompositeMeshes();
+			FCompositeRenderTargetPool::Get().ReleaseTarget(SceneRenderTarget);
 		}
+
+		CachedValidScenePasses = NumValidScenePasses;
+		bPassesHaveChanged = true;
+	}
+
+	if (bPassesHaveChanged)
+	{
+		UpdateCompositeMeshes();
 	}
 }
 
@@ -241,6 +235,16 @@ void UCompositeLayerPlate::SetPlateMode(ECompositePlateMode InPlateMode)
 	PlateMode = InPlateMode;
 }
 
+void UCompositeLayerPlate::SetIsHoldoutEnabled(bool bInIsHoldoutEnabled)
+{
+	bIsHoldoutEnabled = bInIsHoldoutEnabled;
+
+	if (!bIsHoldoutEnabled)
+	{
+		UnregisterCompositePrimitives();
+	}
+}
+
 UTexture* UCompositeLayerPlate::GetCompositeTexture() const
 {
 	if (CachedValidScenePasses > 0)
@@ -280,11 +284,7 @@ void UCompositeLayerPlate::PreEditChange(FProperty* PropertyThatWillChange)
 	{
 		PreEditCompositeMeshes = CompositeMeshes;
 
-		UCompositeCoreSubsystem* Subsystem = UWorld::GetSubsystem<UCompositeCoreSubsystem>(GetWorld());
-		if (IsValid(Subsystem))
-		{
-			Subsystem->UnregisterPrimitives(GetPrimitives());
-		}
+		UnregisterCompositePrimitives();
 	}
 	else if (PropertyName == GET_MEMBER_NAME_CHECKED(ThisClass, Texture))
 	{
@@ -334,12 +334,8 @@ void UCompositeLayerPlate::PostEditChangeProperty(FPropertyChangedEvent& Propert
 			}
 		}
 
-		UCompositeCoreSubsystem* Subsystem = UWorld::GetSubsystem<UCompositeCoreSubsystem>(GetWorld());
-		if (IsValid(Subsystem) && IsRendering())
-		{
-			// Register primitives to be marked as holdout and rendered in the built-in composite custom render pass.
-			Subsystem->RegisterPrimitives(GetPrimitives());
-		}
+		// TODO: Remove and rely solely on the tick registration?
+		TryRegisterCompositePrimitives();
 
 		PreEditCompositeMeshes.Reset();
 	}
@@ -347,6 +343,10 @@ void UCompositeLayerPlate::PostEditChangeProperty(FPropertyChangedEvent& Propert
 	{
 		TryOpenMediaProfileSource();
 	}
+	else if (PropertyName == GET_MEMBER_NAME_CHECKED(ThisClass, bIsHoldoutEnabled))
+	{
+		SetIsHoldoutEnabled(bIsHoldoutEnabled);
+	}
 
 	// For simplicity, we make sure to update composite meshes on any property change.
 	UpdateCompositeMeshes();
@@ -377,12 +377,8 @@ void UCompositeLayerPlate::PostTransacted(const FTransactionObjectEvent& InTrans
 		{
 			UpdateCompositeMeshes();
 
-			UCompositeCoreSubsystem* Subsystem = UWorld::GetSubsystem<UCompositeCoreSubsystem>(GetWorld());
-			if (IsValid(Subsystem) && IsRendering())
-			{
-				// Register primitives to be marked as holdout and rendered in the built-in composite custom render pass.
-				Subsystem->RegisterPrimitives(GetPrimitives());
-			}
+			// TODO: Remove and rely solely on the tick registration?
+			TryRegisterCompositePrimitives();
 		}
 	}
 }
@@ -705,6 +701,60 @@ bool UCompositeLayerPlate::IsUsingAfterDOFTranslucency(UMaterialInterface* Mater
 	return false;
 }
 
+void UCompositeLayerPlate::TryRegisterCompositePrimitives(const UWorld* SourceWorld) const
+{
+	if (!IsRendering())
+	{
+		return;
+	}
+
+	TArray<UPrimitiveComponent*> Primitives = GetPrimitives();
+
+	if (bIsHoldoutEnabled)
+	{
+		for (UPrimitiveComponent* Primitive : Primitives)
+		{
+			if (Primitive && !Primitive->bHoldout)
+			{
+				Primitive->Modify();
+				Primitive->SetHoldout(true);
+			}
+		}
+	}
+
+	if (PlateMode == ECompositePlateMode::CompositeMesh)
+	{
+		UCompositeCoreSubsystem* Subsystem = UWorld::GetSubsystem<UCompositeCoreSubsystem>(SourceWorld ? SourceWorld : GetWorld());
+		if (Subsystem)
+		{
+			// Primitives to be marked as holdout and rendered in the built-in composite custom render pass.
+			Subsystem->RegisterPrimitives(GetPrimitives(), ECompositeCoreHoldoutManagement::None);
+		}
+	}
+}
+
+void UCompositeLayerPlate::UnregisterCompositePrimitives(const UWorld* SourceWorld) const
+{
+	TArray<UPrimitiveComponent*> Primitives = GetPrimitives();
+	
+	// Explicitely disable holdout since the subsystem below is only used in ECompositePlateMode::CompositeMesh. 
+	for (UPrimitiveComponent* Primitive : Primitives)
+	{
+		if (Primitive && Primitive->bHoldout)
+		{
+			Primitive->Modify();
+			Primitive->SetHoldout(false);
+		}
+	}
+
+	// Always unregister no matter the current PlateMode
+	UCompositeCoreSubsystem* Subsystem = UWorld::GetSubsystem<UCompositeCoreSubsystem>(SourceWorld ? SourceWorld : GetWorld());
+	if (IsValid(Subsystem))
+	{
+		Subsystem->UnregisterPrimitives(Primitives, ECompositeCoreHoldoutManagement::None);
+	}
+}
+
 void UCompositeLayerPlate::UpdateCompositeMeshes(bool bRegisterOnly) const
 {
 	for (const TSoftObjectPtr<AActor>& MeshActor : CompositeMeshes)
@@ -790,8 +840,6 @@ void UCompositeLayerPlate::UpdatePrimitiveComponent(UPrimitiveComponent& InPrimi
 
 void UCompositeLayerPlate::PropagateStateChange(bool bApply, const UWorld* SourceWorld) const
 {
-	UCompositeCoreSubsystem* Subsystem = UWorld::GetSubsystem<UCompositeCoreSubsystem>(SourceWorld ? SourceWorld : GetWorld());
-
 	if (bApply && IsRendering())
 	{
 		if (CompositeMeshes.IsEmpty())
@@ -802,11 +850,8 @@ void UCompositeLayerPlate::PropagateStateChange(bool bApply, const UWorld* Sourc
 		constexpr bool bRegisterOnly = true;
 		UpdateCompositeMeshes(bRegisterOnly);
 
-		if (IsValid(Subsystem))
-		{
-			// Rely on the scene view extension to mark primivites as holdout, and render them separately.
-			Subsystem->RegisterPrimitives(GetPrimitives());
-		}
+		// TODO: Remove and rely solely on the tick registration?
+		TryRegisterCompositePrimitives(SourceWorld);
 
 		// Conditionally override translucency console variables
 		for (const TSoftObjectPtr<AActor>& MeshActor : CompositeMeshes)
@@ -841,10 +886,7 @@ void UCompositeLayerPlate::PropagateStateChange(bool bApply, const UWorld* Sourc
 		{
 			UE::Composite::Private::RemoveCompositeAllAssetUserData(CompositeMeshes);
 
-			if (IsValid(Subsystem))
-			{
-				Subsystem->UnregisterPrimitives(GetPrimitives());
-			}
+			UnregisterCompositePrimitives(SourceWorld);
 		}
 
 		constexpr bool bOverride = false;
diff --git a/Engine/Plugins/Compositing/Composite/Source/Composite/Public/Layers/CompositeLayerPlate.h b/Engine/Plugins/Compositing/Composite/Source/Composite/Public/Layers/CompositeLayerPlate.h
--- a/Engine/Plugins/Compositing/Composite/Source/Composite/Public/Layers/CompositeLayerPlate.h
+++ b/Engine/Plugins/Compositing/Composite/Source/Composite/Public/Layers/CompositeLayerPlate.h
@@ -26,8 +26,13 @@ enum class ECompositePlateMode : uint8
 /**
 * Convenience layer for media texture integration, with additional processing options.
 * 
-* With composite meshes and default plugin materials, the media texture is projected into the scene & rendered separately as a custom render pass.
-* The result is dilated and combined with other layers such as the main render (where the same composite meshes are rendered as holdouts).
+* With composite meshes and default plugin materials, the media texture is projected
+* into the scene. These meshes are rendered separately in a custom render pass without
+* post-processing. By default, they are also marked as holdout such that the plate layer
+* can be composited under the main render.
+* 
+* Without composite meshes, the plate layer can be used solely as a 2D layer. In this case
+* however, interactions between CG and media in the scene are lost.
 */
 UCLASS(MinimalAPI, BlueprintType, Blueprintable, EditInlineNew, CollapseCategories, meta = (DisplayName = "Plate"))
 class UCompositeLayerPlate : public UCompositeLayerBase, public FTickableGameObject
@@ -100,6 +105,14 @@ public:
 	UFUNCTION(BlueprintSetter)
 	UE_API void SetPlateMode(ECompositePlateMode InPlateMode);
 
+	/** Get whether primitive alpha holdout is active for composite meshes in the main render. */
+	UFUNCTION(BlueprintGetter)
+	UE_API bool GetIsHoldoutEnabled() const { return bIsHoldoutEnabled; }
+
+	/** Set whether primitive alpha holdout is active for composite meshes in the main render. */
+	UFUNCTION(BlueprintSetter)
+	UE_API void SetIsHoldoutEnabled(bool bInIsHoldoutEnabled);
+
 	/** Returns the media texture, optionally pre-processed with layer & scene-only passes into a render target. */
 	UFUNCTION(BlueprintCallable, Category = "Composite")
 	UTexture* GetCompositeTexture() const;
@@ -143,7 +156,23 @@ private:
 	UPROPERTY(EditAnywhere, BlueprintGetter = GetPlateMode, BlueprintSetter = SetPlateMode, Category = "", AdvancedDisplay, meta = (AllowPrivateAccess = true, DisplayName = "Mode"))
 	ECompositePlateMode PlateMode;
 
+	/**
+	* Composite meshes are marked as primitive alpha holdout for the main render.
+	* 
+	* This is default-enabled to bypass post-processing on media pixels, primarily to avoid temporal upscaling
+	* artifacts and the tone curve. However when working in 3D, holdout can be disabled to only preserve the
+	* plate projection functionality and media passes.
+	*/
+	UPROPERTY(EditAnywhere, BlueprintGetter = GetIsHoldoutEnabled, BlueprintSetter = SetIsHoldoutEnabled, Category = "", AdvancedDisplay, meta = (AllowPrivateAccess = true))
+	bool bIsHoldoutEnabled;
+
 private:
+	/** Condtionally register composite mesh primitives: active holdout and add to the scene view extension custom render pass. */
+	void TryRegisterCompositePrimitives(const UWorld* SourceWorld = nullptr) const;
+
+	/** Unregister composite mesh primitives: deactive holdout and remove from the scene view extension custom render pass. */
+	void UnregisterCompositePrimitives(const UWorld* SourceWorld = nullptr) const;
+
 	/** Update the (media) texture on all composite meshes. We also add a user asset data to track changes. */
 	void UpdateCompositeMeshes(bool bRegisterOnly = false) const;
 	
diff --git a/Engine/Plugins/Compositing/CompositeCore/Source/CompositeCore/Private/CompositeCoreSceneViewExtension.cpp b/Engine/Plugins/Compositing/CompositeCore/Source/CompositeCore/Private/CompositeCoreSceneViewExtension.cpp
--- a/Engine/Plugins/Compositing/CompositeCore/Source/CompositeCore/Private/CompositeCoreSceneViewExtension.cpp
+++ b/Engine/Plugins/Compositing/CompositeCore/Source/CompositeCore/Private/CompositeCoreSceneViewExtension.cpp
@@ -83,7 +83,7 @@ FCompositeCoreSceneViewExtension::FCompositeCoreSceneViewExtension(const FAutoRe
 
 FCompositeCoreSceneViewExtension::~FCompositeCoreSceneViewExtension() = default;
 
-void FCompositeCoreSceneViewExtension::RegisterPrimitives_GameThread(const TArray<UPrimitiveComponent*>& InPrimitiveComponents)
+void FCompositeCoreSceneViewExtension::RegisterPrimitives_GameThread(const TArray<UPrimitiveComponent*>& InPrimitiveComponents, ECompositeCoreHoldoutManagement HoldoutManagement)
 {
 	check(IsInGameThread());
 
@@ -99,7 +99,7 @@ void FCompositeCoreSceneViewExtension::RegisterPrimitives_GameThread(const TArra
 			CompositePrimitives.Add(InPrimitiveComponent);
 		}
 
-		if (!InPrimitiveComponent->bHoldout)
+		if (HoldoutManagement == ECompositeCoreHoldoutManagement::Automatic && !InPrimitiveComponent->bHoldout)
 		{
 			InPrimitiveComponent->Modify();
 			InPrimitiveComponent->SetHoldout(true);
@@ -107,7 +107,7 @@ void FCompositeCoreSceneViewExtension::RegisterPrimitives_GameThread(const TArra
 	}
 }
 
-void FCompositeCoreSceneViewExtension::UnregisterPrimitives_GameThread(const TArray<UPrimitiveComponent*>& InPrimitiveComponents)
+void FCompositeCoreSceneViewExtension::UnregisterPrimitives_GameThread(const TArray<UPrimitiveComponent*>& InPrimitiveComponents, ECompositeCoreHoldoutManagement HoldoutManagement)
 {
 	check(IsInGameThread());
 
@@ -123,7 +123,7 @@ void FCompositeCoreSceneViewExtension::UnregisterPrimitives_GameThread(const TAr
 			CompositePrimitives.Remove(InPrimitiveComponent);
 		}
 
-		if (InPrimitiveComponent->bHoldout)
+		if (HoldoutManagement == ECompositeCoreHoldoutManagement::Automatic && InPrimitiveComponent->bHoldout)
 		{
 			InPrimitiveComponent->Modify();
 			InPrimitiveComponent->SetHoldout(false);
diff --git a/Engine/Plugins/Compositing/CompositeCore/Source/CompositeCore/Private/CompositeCoreSceneViewExtension.h b/Engine/Plugins/Compositing/CompositeCore/Source/CompositeCore/Private/CompositeCoreSceneViewExtension.h
--- a/Engine/Plugins/Compositing/CompositeCore/Source/CompositeCore/Private/CompositeCoreSceneViewExtension.h
+++ b/Engine/Plugins/Compositing/CompositeCore/Source/CompositeCore/Private/CompositeCoreSceneViewExtension.h
@@ -9,6 +9,7 @@
 #include "RendererInterface.h"
 
 #include "CompositeCoreSettings.h"
+#include "CompositeCoreTypes.h"
 #include "Passes/CompositeCorePassProxy.h"
 
 class FViewInfo;
@@ -67,10 +68,10 @@ public:
 	~FCompositeCoreSceneViewExtension();
 
 	/** Register primitives for compositing. */
-	void RegisterPrimitives_GameThread(const TArray<UPrimitiveComponent*>& InPrimitiveComponents);
+	void RegisterPrimitives_GameThread(const TArray<UPrimitiveComponent*>& InPrimitiveComponents, ECompositeCoreHoldoutManagement HoldoutManagement = ECompositeCoreHoldoutManagement::Automatic);
 
 	/** Unregister primitives for compositing. */
-	void UnregisterPrimitives_GameThread(const TArray<UPrimitiveComponent*>& InPrimitiveComponents);
+	void UnregisterPrimitives_GameThread(const TArray<UPrimitiveComponent*>& InPrimitiveComponents, ECompositeCoreHoldoutManagement HoldoutManagement = ECompositeCoreHoldoutManagement::Automatic);
 
 	/** Set frame render work to be enqueued on the render thread. */
 	void SetRenderWork_GameThread(UE::CompositeCore::FRenderWork&& InWork);
diff --git a/Engine/Plugins/Compositing/CompositeCore/Source/CompositeCore/Private/CompositeCoreSubsystem.cpp b/Engine/Plugins/Compositing/CompositeCore/Source/CompositeCore/Private/CompositeCoreSubsystem.cpp
--- a/Engine/Plugins/Compositing/CompositeCore/Source/CompositeCore/Private/CompositeCoreSubsystem.cpp
+++ b/Engine/Plugins/Compositing/CompositeCore/Source/CompositeCore/Private/CompositeCoreSubsystem.cpp
@@ -75,35 +75,35 @@ void UCompositeCoreSubsystem::Tick(float DeltaTime)
 	Super::Tick(DeltaTime);
 }
 
-void UCompositeCoreSubsystem::RegisterPrimitive(UPrimitiveComponent* InPrimitiveComponent)
+void UCompositeCoreSubsystem::RegisterPrimitive(UPrimitiveComponent* InPrimitiveComponent, ECompositeCoreHoldoutManagement HoldoutManagement)
 {
 	TArray<UPrimitiveComponent*> PrimitiveComponents = { InPrimitiveComponent };
 
-	RegisterPrimitives(PrimitiveComponents);
+	RegisterPrimitives(PrimitiveComponents, HoldoutManagement);
 }
 
-void UCompositeCoreSubsystem::RegisterPrimitives(const TArray<UPrimitiveComponent*>& InPrimitiveComponents)
+void UCompositeCoreSubsystem::RegisterPrimitives(const TArray<UPrimitiveComponent*>& InPrimitiveComponents, ECompositeCoreHoldoutManagement HoldoutManagement)
 {
 	const bool bValidSettings = ValidateProjectSettings();
 
 	if (bValidSettings && CompositeCoreViewExtension.IsValid())
 	{
-		CompositeCoreViewExtension->RegisterPrimitives_GameThread(InPrimitiveComponents);
+		CompositeCoreViewExtension->RegisterPrimitives_GameThread(InPrimitiveComponents, HoldoutManagement);
 	}
 }
 
-void UCompositeCoreSubsystem::UnregisterPrimitive(UPrimitiveComponent* InPrimitiveComponent)
+void UCompositeCoreSubsystem::UnregisterPrimitive(UPrimitiveComponent* InPrimitiveComponent, ECompositeCoreHoldoutManagement HoldoutManagement)
 {
 	TArray<UPrimitiveComponent*> PrimitiveComponents = { InPrimitiveComponent };
 
-	UnregisterPrimitives(PrimitiveComponents);
+	UnregisterPrimitives(PrimitiveComponents, HoldoutManagement);
 }
 
-void UCompositeCoreSubsystem::UnregisterPrimitives(const TArray<UPrimitiveComponent*>& InPrimitiveComponents)
+void UCompositeCoreSubsystem::UnregisterPrimitives(const TArray<UPrimitiveComponent*>& InPrimitiveComponents, ECompositeCoreHoldoutManagement HoldoutManagement)
 {
 	if (CompositeCoreViewExtension.IsValid())
 	{
-		CompositeCoreViewExtension->UnregisterPrimitives_GameThread(InPrimitiveComponents);
+		CompositeCoreViewExtension->UnregisterPrimitives_GameThread(InPrimitiveComponents, HoldoutManagement);
 	}
 }
 
diff --git a/Engine/Plugins/Compositing/CompositeCore/Source/CompositeCore/Public/CompositeCoreSubsystem.h b/Engine/Plugins/Compositing/CompositeCore/Source/CompositeCore/Public/CompositeCoreSubsystem.h
--- a/Engine/Plugins/Compositing/CompositeCore/Source/CompositeCore/Public/CompositeCoreSubsystem.h
+++ b/Engine/Plugins/Compositing/CompositeCore/Source/CompositeCore/Public/CompositeCoreSubsystem.h
@@ -5,6 +5,7 @@
 #include "MaterialShared.h"
 #include "Subsystems/WorldSubsystem.h"
 
+#include "CompositeCoreTypes.h"
 #include "CompositeCoreSubsystem.generated.h"
 
 #define UE_API COMPOSITECORE_API
@@ -45,17 +46,17 @@ public:
 
 	/* Register a single primitive for compositing. */
 	UFUNCTION(BlueprintCallable, Category = "Composite Core")
-	UE_API void RegisterPrimitive(UPrimitiveComponent* InPrimitiveComponent);
+	UE_API void RegisterPrimitive(UPrimitiveComponent* InPrimitiveComponent, ECompositeCoreHoldoutManagement HoldoutManagement = ECompositeCoreHoldoutManagement::Automatic);
 
 	/* Register multiple primitives for compositing. */
-	UE_API void RegisterPrimitives(const TArray<UPrimitiveComponent*>& InPrimitiveComponents);
+	UE_API void RegisterPrimitives(const TArray<UPrimitiveComponent*>& InPrimitiveComponents, ECompositeCoreHoldoutManagement HoldoutManagement = ECompositeCoreHoldoutManagement::Automatic);
 
 	/* Unregister a single primitive from compositing. */
 	UFUNCTION(BlueprintCallable, Category = "Composite Core")
-	UE_API void UnregisterPrimitive(UPrimitiveComponent* InPrimitiveComponent);
+	UE_API void UnregisterPrimitive(UPrimitiveComponent* InPrimitiveComponent, ECompositeCoreHoldoutManagement HoldoutManagement = ECompositeCoreHoldoutManagement::Automatic);
 
 	/* Unregister multiple primitives from compositing. */
-	UE_API void UnregisterPrimitives(const TArray<UPrimitiveComponent*>& InPrimitiveComponents);
+	UE_API void UnregisterPrimitives(const TArray<UPrimitiveComponent*>& InPrimitiveComponents, ECompositeCoreHoldoutManagement HoldoutManagement = ECompositeCoreHoldoutManagement::Automatic);
 
 	/* Set scene view extension frame render work. */
 	UE_API void SetRenderWork(UE::CompositeCore::FRenderWork&& InWork);
diff --git a/Engine/Plugins/Compositing/CompositeCore/Source/CompositeCore/Public/CompositeCoreTypes.h b/Engine/Plugins/Compositing/CompositeCore/Source/CompositeCore/Public/CompositeCoreTypes.h
new file mode 100644
--- /dev/null
+++ b/Engine/Plugins/Compositing/CompositeCore/Source/CompositeCore/Public/CompositeCoreTypes.h
@@ -0,0 +1,20 @@
+// Copyright Epic Games, Inc. All Rights Reserved.
+
+#pragma once
+
+#include "CoreTypes.h"
+#include "UObject/ObjectMacros.h"
+
+#include "CompositeCoreTypes.generated.h"
+
+#define UE_API COMPOSITECORE_API
+
+/* Holdout property management on registered primitives. */
+UENUM(BlueprintType)
+enum class ECompositeCoreHoldoutManagement : uint8
+{
+	None,
+	Automatic,
+};
+
+#undef UE_API
