Composure: Automatically hide registered composite meshes when a plate layer or a composite actor is disabled, deactivated or hidden in editor.

* This yields more intuitive behavior, and allows for switching between 2D plate layers & 3D plate layers more easily.
* Convenience remains automatically functional with multiple composite actors.
* If undesirable, this behavior can be deactivated with the Composite.Plate.ControlMeshVisibility cvar.
* Removed redundant primitive registration calls.

#Synced-CL 50671255

diff --git a/Engine/Plugins/Compositing/Composite/Source/Composite/Private/CompositeActor.cpp b/Engine/Plugins/Compositing/Composite/Source/Composite/Private/CompositeActor.cpp
--- a/Engine/Plugins/Compositing/Composite/Source/Composite/Private/CompositeActor.cpp
+++ b/Engine/Plugins/Compositing/Composite/Source/Composite/Private/CompositeActor.cpp
@@ -623,9 +623,7 @@ void ACompositeActor::SetActive(bool bInActive)
 {
 	bIsActive = bInActive;
 
-	PropagateStateChange(IsRendering());
-
-	// Deactivate all other composite actors
+	// Deactivate all other composite actors first
 	const UWorld* World = GetWorld();
 
 	if (bIsActive && IsValid(World))
@@ -640,6 +638,8 @@ void ACompositeActor::SetActive(bool bInActive)
 		}
 	}
 
+	PropagateStateChange(IsRendering());
+
 	// Update the view projection matrix sent to the material parameter collection when changing which actor is active.
 	if (bIsActive && IsValid(ViewProjectionComponent))
 	{
diff --git a/Engine/Plugins/Compositing/Composite/Source/Composite/Private/Layers/CompositeLayerPlate.cpp b/Engine/Plugins/Compositing/Composite/Source/Composite/Private/Layers/CompositeLayerPlate.cpp
--- a/Engine/Plugins/Compositing/Composite/Source/Composite/Private/Layers/CompositeLayerPlate.cpp
+++ b/Engine/Plugins/Compositing/Composite/Source/Composite/Private/Layers/CompositeLayerPlate.cpp
@@ -11,6 +11,7 @@
 #include "CompositeRenderTargetPool.h"
 #include "Engine/Texture2D.h"
 #include "MediaTexture.h"
+#include "EngineUtils.h"
 #include "Engine/TextureRenderTarget2D.h"
 #include "Materials/Material.h"
 #include "Materials/MaterialInstanceDynamic.h"
@@ -31,6 +32,12 @@
 
 const FLazyName UCompositeLayerPlate::CompositeTextureName = FLazyName(TEXT("CompositeTexture"));
 
+TAutoConsoleVariable<bool> CVarCompositePlateControlMeshVisibility(
+	TEXT("Composite.Plate.ControlMeshVisibility"),
+	true,
+	TEXT("When enabled, the plate layers automatically control the visibility of registered composite meshes")
+);
+
 #define LOCTEXT_NAMESPACE "Composite"
 
 namespace UE
@@ -114,7 +121,7 @@ void UCompositeLayerPlate::Tick(float DeltaTime)
 	}
 
 	// Update holdout on tick in case composite meshes dynamically recreate their primitives
-	TryRegisterCompositePrimitives();
+	TryRegisterCompositePrimitives(GetPrimitives());
 
 	// Constantly check if passes have changed to manage render targets & texture bindings on composite meshes.
 	const int32 NumValidMediaPasses = GetValidPassesNum(MediaPasses);
@@ -241,7 +248,7 @@ void UCompositeLayerPlate::SetIsHoldoutEnabled(bool bInIsHoldoutEnabled)
 
 	if (!bIsHoldoutEnabled)
 	{
-		UnregisterCompositePrimitives();
+		UnregisterCompositePrimitives(GetPrimitives());
 	}
 }
 
@@ -266,7 +273,22 @@ UTexture* UCompositeLayerPlate::GetCompositeTexture() const
 #if WITH_EDITOR
 bool UCompositeLayerPlate::CanEditChange(const FProperty* InProperty) const
 {
-	return Super::CanEditChange(InProperty);
+	bool bIsEditable = Super::CanEditChange(InProperty);
+
+	if (bIsEditable && InProperty)
+	{
+		const FName PropertyName = InProperty->GetFName();
+
+		if (PropertyName == GET_MEMBER_NAME_CHECKED(ThisClass, bIsHoldoutEnabled))
+		{
+			if (CompositeMeshes.IsEmpty())
+			{
+				bIsEditable = false;
+			}
+		}
+	}
+
+	return bIsEditable;
 }
 
 void UCompositeLayerPlate::PreEditChange(FProperty* PropertyThatWillChange)
@@ -284,7 +306,7 @@ void UCompositeLayerPlate::PreEditChange(FProperty* PropertyThatWillChange)
 	{
 		PreEditCompositeMeshes = CompositeMeshes;
 
-		UnregisterCompositePrimitives();
+		UnregisterCompositePrimitives(GetPrimitives());
 	}
 	else if (PropertyName == GET_MEMBER_NAME_CHECKED(ThisClass, Texture))
 	{
@@ -334,9 +356,6 @@ void UCompositeLayerPlate::PostEditChangeProperty(FPropertyChangedEvent& Propert
 			}
 		}
 
-		// TODO: Remove and rely solely on the tick registration?
-		TryRegisterCompositePrimitives();
-
 		PreEditCompositeMeshes.Reset();
 	}
 	else if (PropertyName == GET_MEMBER_NAME_CHECKED(ThisClass, Texture))
@@ -376,9 +395,6 @@ void UCompositeLayerPlate::PostTransacted(const FTransactionObjectEvent& InTrans
 		else if (ChangedProperties.Contains(GET_MEMBER_NAME_CHECKED(ThisClass, CompositeMeshes)))
 		{
 			UpdateCompositeMeshes();
-
-			// TODO: Remove and rely solely on the tick registration?
-			TryRegisterCompositePrimitives();
 		}
 	}
 }
@@ -701,18 +717,16 @@ bool UCompositeLayerPlate::IsUsingAfterDOFTranslucency(UMaterialInterface* Mater
 	return false;
 }
 
-void UCompositeLayerPlate::TryRegisterCompositePrimitives(const UWorld* SourceWorld) const
+void UCompositeLayerPlate::TryRegisterCompositePrimitives(const TArray<UPrimitiveComponent*>& InPrimitives, const UWorld* SourceWorld) const
 {
 	if (!IsRendering())
 	{
 		return;
 	}
 
-	TArray<UPrimitiveComponent*> Primitives = GetPrimitives();
-
 	if (bIsHoldoutEnabled)
 	{
-		for (UPrimitiveComponent* Primitive : Primitives)
+		for (UPrimitiveComponent* Primitive : InPrimitives)
 		{
 			if (Primitive && !Primitive->bHoldout)
 			{
@@ -728,17 +742,15 @@ void UCompositeLayerPlate::TryRegisterCompositePrimitives(const UWorld* SourceWo
 		if (Subsystem)
 		{
 			// Primitives to be marked as holdout and rendered in the built-in composite custom render pass.
-			Subsystem->RegisterPrimitives(GetPrimitives(), ECompositeCoreHoldoutManagement::None);
+			Subsystem->RegisterPrimitives(InPrimitives, ECompositeCoreHoldoutManagement::None);
 		}
 	}
 }
 
-void UCompositeLayerPlate::UnregisterCompositePrimitives(const UWorld* SourceWorld) const
-{
-	TArray<UPrimitiveComponent*> Primitives = GetPrimitives();
-	
+void UCompositeLayerPlate::UnregisterCompositePrimitives(const TArray<UPrimitiveComponent*>& InPrimitives, const UWorld* SourceWorld) const
+{	
 	// Explicitely disable holdout since the subsystem below is only used in ECompositePlateMode::CompositeMesh. 
-	for (UPrimitiveComponent* Primitive : Primitives)
+	for (UPrimitiveComponent* Primitive : InPrimitives)
 	{
 		if (Primitive && Primitive->bHoldout)
 		{
@@ -751,7 +763,7 @@ void UCompositeLayerPlate::UnregisterCompositePrimitives(const UWorld* SourceWor
 	UCompositeCoreSubsystem* Subsystem = UWorld::GetSubsystem<UCompositeCoreSubsystem>(SourceWorld ? SourceWorld : GetWorld());
 	if (IsValid(Subsystem))
 	{
-		Subsystem->UnregisterPrimitives(Primitives, ECompositeCoreHoldoutManagement::None);
+		Subsystem->UnregisterPrimitives(InPrimitives, ECompositeCoreHoldoutManagement::None);
 	}
 }
 
@@ -840,6 +852,8 @@ void UCompositeLayerPlate::UpdatePrimitiveComponent(UPrimitiveComponent& InPrimi
 
 void UCompositeLayerPlate::PropagateStateChange(bool bApply, const UWorld* SourceWorld) const
 {
+	const bool bControlMeshVisibility = CVarCompositePlateControlMeshVisibility.GetValueOnGameThread();
+
 	if (bApply && IsRendering())
 	{
 		if (CompositeMeshes.IsEmpty())
@@ -850,8 +864,21 @@ void UCompositeLayerPlate::PropagateStateChange(bool bApply, const UWorld* Sourc
 		constexpr bool bRegisterOnly = true;
 		UpdateCompositeMeshes(bRegisterOnly);
 
-		// TODO: Remove and rely solely on the tick registration?
-		TryRegisterCompositePrimitives(SourceWorld);
+		const TArray<UPrimitiveComponent*> Primitives = GetPrimitives();
+
+		if (bControlMeshVisibility)
+		{
+			for (UPrimitiveComponent* Primitive : Primitives)
+			{
+				Primitive->Modify();
+
+				constexpr bool bNewVisibility = true;
+				Primitive->SetVisibility(bNewVisibility);
+			}
+		}
+
+		// Note: We could rely solely on the tick registration, but this call is kept for symmetry.
+		TryRegisterCompositePrimitives(Primitives, SourceWorld);
 
 		// Conditionally override translucency console variables
 		for (const TSoftObjectPtr<AActor>& MeshActor : CompositeMeshes)
@@ -886,7 +913,26 @@ void UCompositeLayerPlate::PropagateStateChange(bool bApply, const UWorld* Sourc
 		{
 			UE::Composite::Private::RemoveCompositeAllAssetUserData(CompositeMeshes);
 
-			UnregisterCompositePrimitives(SourceWorld);
+			const TArray<UPrimitiveComponent*> Primitives = GetPrimitives();
+
+			UnregisterCompositePrimitives(Primitives, SourceWorld);
+
+			if (bControlMeshVisibility)
+			{
+				TSet<UPrimitiveComponent*> UsedPrimitives = GetPrimitivesUsedByOtherActors();
+
+				for (UPrimitiveComponent* Primitive : Primitives)
+				{
+					// Only make a primitive invisible if it's not already used by another composite actor
+					if (!UsedPrimitives.Contains(Primitive))
+					{
+						Primitive->Modify();
+
+						constexpr bool bNewVisibility = false;
+						Primitive->SetVisibility(bNewVisibility);
+					}
+				}
+			}
 		}
 
 		constexpr bool bOverride = false;
@@ -917,6 +963,50 @@ bool UCompositeLayerPlate::IsCompositeMeshActorAlreadyInUse(TSoftObjectPtr<AActo
 	return false;
 }
 
+TSet<UPrimitiveComponent*> UCompositeLayerPlate::GetPrimitivesUsedByOtherActors() const
+{
+	// Note: as oppposed to trying to maintain a static cache of actively used primitives,
+	// we simply query the world state and pay the lookup cost.
+
+	TSet<UPrimitiveComponent*> Output;
+
+	const UWorld* World = GetWorld();
+	const ACompositeActor* ThisActor = GetTypedOuter<ACompositeActor>();
+
+	if (World && ThisActor)
+	{
+		for (TActorIterator<ACompositeActor> It(World); It; ++It)
+		{
+			const ACompositeActor* OtherActor = Cast<ACompositeActor>(*It);
+
+			if (!OtherActor || OtherActor == ThisActor || !OtherActor->IsRendering())
+			{
+				continue;
+			}
+			
+			for (const TObjectPtr<UCompositeLayerBase>& OtherLayer : OtherActor->GetCompositeLayers())
+			{
+				const UCompositeLayerPlate* OtherPlateLayer = Cast<UCompositeLayerPlate>(OtherLayer.Get());
+
+				if (!OtherPlateLayer || !ensure(OtherPlateLayer != this) || !OtherPlateLayer->IsRendering())
+				{
+					continue;
+				}
+
+				for (UPrimitiveComponent* OtherPrimitive : OtherPlateLayer->GetPrimitives())
+				{
+					if (OtherPrimitive)
+					{
+						Output.Add(OtherPrimitive);
+					}
+				}
+			}
+		}
+	}
+
+	return Output;
+}
+
 TObjectPtr<UTextureRenderTarget2D> UCompositeLayerPlate::GetOrCreateRenderTarget(TObjectPtr<UTextureRenderTarget2D>& InRenderTarget) const
 {
 	if (!IsValid(InRenderTarget))
diff --git a/Engine/Plugins/Compositing/Composite/Source/Composite/Public/Layers/CompositeLayerPlate.h b/Engine/Plugins/Compositing/Composite/Source/Composite/Public/Layers/CompositeLayerPlate.h
--- a/Engine/Plugins/Compositing/Composite/Source/Composite/Public/Layers/CompositeLayerPlate.h
+++ b/Engine/Plugins/Compositing/Composite/Source/Composite/Public/Layers/CompositeLayerPlate.h
@@ -168,10 +168,10 @@ private:
 
 private:
 	/** Condtionally register composite mesh primitives: active holdout and add to the scene view extension custom render pass. */
-	void TryRegisterCompositePrimitives(const UWorld* SourceWorld = nullptr) const;
+	void TryRegisterCompositePrimitives(const TArray<UPrimitiveComponent*>& InPrimitives, const UWorld* SourceWorld = nullptr) const;
 
 	/** Unregister composite mesh primitives: deactive holdout and remove from the scene view extension custom render pass. */
-	void UnregisterCompositePrimitives(const UWorld* SourceWorld = nullptr) const;
+	void UnregisterCompositePrimitives(const TArray<UPrimitiveComponent*>& InPrimitives, const UWorld* SourceWorld = nullptr) const;
 
 	/** Update the (media) texture on all composite meshes. We also add a user asset data to track changes. */
 	void UpdateCompositeMeshes(bool bRegisterOnly = false) const;
@@ -190,6 +190,9 @@ private:
 	/** Convenience function verify if the composite mesh actor is already in use by another layer. */
 	bool IsCompositeMeshActorAlreadyInUse(TSoftObjectPtr<AActor> InCompositeMeshActor) const;
 
+	/** Get set of primitives already in use by other actors. */
+	TSet<UPrimitiveComponent*> GetPrimitivesUsedByOtherActors() const;
+
 	/** Convenience function to get or create the specified render target at composite actor resolution. */
 	TObjectPtr<UTextureRenderTarget2D> GetOrCreateRenderTarget(TObjectPtr<UTextureRenderTarget2D>& InRenderTarget) const;
 
