Composite: Minor change to large bounds enforcement on composite depth mesh to allow for WPO.

#Synced-CL 50608378

diff --git a/Engine/Plugins/Compositing/Composite/Source/Composite/Private/Components/CompositeDepthMeshComponent.cpp b/Engine/Plugins/Compositing/Composite/Source/Composite/Private/Components/CompositeDepthMeshComponent.cpp
--- a/Engine/Plugins/Compositing/Composite/Source/Composite/Private/Components/CompositeDepthMeshComponent.cpp
+++ b/Engine/Plugins/Compositing/Composite/Source/Composite/Private/Components/CompositeDepthMeshComponent.cpp
@@ -104,9 +104,6 @@ UCompositeDepthMeshComponent::UCompositeDepthMeshComponent(const FObjectInitiali
 	SetCollisionEnabled(ECollisionEnabled::NoCollision);
 	SetCastShadow(false);
 	SetEvaluateWorldPositionOffset(true);
-	
-	// Increase bounds to prevent flickering with WPO material
-	BoundsScale = 1000.0f;
 }
 
 void UCompositeDepthMeshComponent::PostLoad()
@@ -119,6 +116,12 @@ void UCompositeDepthMeshComponent::PostLoad()
 	}
 }
 
+FBoxSphereBounds UCompositeDepthMeshComponent::CalcBounds(const FTransform& LocalToWorld) const
+{
+	const FVector Origin = LocalToWorld.TransformPosition(FVector::ZeroVector);
+	return FBoxSphereBounds(FBox::BuildAABB(Origin, FVector::One() * UE_LARGE_HALF_WORLD_MAX));
+}
+
 void UCompositeDepthMeshComponent::UpdateMaterial()
 {
 	UMaterialInterface* MaterialInterface = GetMaterial(0);
diff --git a/Engine/Plugins/Compositing/Composite/Source/Composite/Public/Components/CompositeDepthMeshComponent.h b/Engine/Plugins/Compositing/Composite/Source/Composite/Public/Components/CompositeDepthMeshComponent.h
--- a/Engine/Plugins/Compositing/Composite/Source/Composite/Public/Components/CompositeDepthMeshComponent.h
+++ b/Engine/Plugins/Compositing/Composite/Source/Composite/Public/Components/CompositeDepthMeshComponent.h
@@ -33,6 +33,9 @@ public:
 	UPROPERTY(EditAnywhere, Category = "Composite")
 	TObjectPtr<UTexture> DepthTexture;
 
+protected:
+	UE_API virtual FBoxSphereBounds CalcBounds(const FTransform& LocalToWorld) const override;
+
 private:
 	/* Convert material to MID and update its depth texture parameter. */
 	void UpdateMaterial();
