Composure: Minor fix to correctly unregister all registered customizations.

#Synced-CL 50742325

diff --git a/Engine/Plugins/Compositing/Composite/Source/CompositeEditor/Private/CompositeEditorModule.cpp b/Engine/Plugins/Compositing/Composite/Source/CompositeEditor/Private/CompositeEditorModule.cpp
--- a/Engine/Plugins/Compositing/Composite/Source/CompositeEditor/Private/CompositeEditorModule.cpp
+++ b/Engine/Plugins/Compositing/Composite/Source/CompositeEditor/Private/CompositeEditorModule.cpp
@@ -297,6 +297,8 @@ void FCompositeEditorModule::UnregisterCustomizations()
 			PropertyModule->UnregisterCustomClassLayout(UCompositeLayerShadowReflection::StaticClass()->GetFName());
 			PropertyModule->UnregisterCustomClassLayout(UCompositeLayerSceneCapture::StaticClass()->GetFName());
 			PropertyModule->UnregisterCustomClassLayout(UCompositeLayerSingleLightShadow::StaticClass()->GetFName());
+			PropertyModule->UnregisterCustomClassLayout(UCompositeDepthMeshComponent::StaticClass()->GetFName());
+			PropertyModule->UnregisterCustomClassLayout(UCompositePassColorKeyer::StaticClass()->GetFName());
 		}
 	}
 }
