Composure: Simplify (de)activation in multiuser with optional VP role.

* If a VP role is specified and does not match the current device role, the composite actor is automatically deactivated.
* We remove the previously automatic de-activation logic when joining multi-user, which was unintuitive and failed to work with newly created actors.

#Synced-CL 50671255

diff --git a/Engine/Plugins/Compositing/Composite/Composite.uplugin b/Engine/Plugins/Compositing/Composite/Composite.uplugin
--- a/Engine/Plugins/Compositing/Composite/Composite.uplugin
+++ b/Engine/Plugins/Compositing/Composite/Composite.uplugin
@@ -44,10 +44,6 @@
 			"Name": "MediaFrameworkUtilities",
 			"Enabled": true
 		},
-		{
-			"Name": "ConcertSyncClient",
-			"Enabled": true
-		},
 		{
 			"Name": "ConcertSyncCore",
 			"Enabled": true
@@ -71,6 +67,10 @@
 			"Name": "EditorScriptingUtilities",
 			"Enabled": true,
 			"TargetAllowList": [ "Editor" ]
+		},
+		{
+			"Name": "VPRoles",
+			"Enabled": true
 		}
 	]
 }
\ No newline at end of file
diff --git a/Engine/Plugins/Compositing/Composite/Source/Composite/Composite.Build.cs b/Engine/Plugins/Compositing/Composite/Source/Composite/Composite.Build.cs
--- a/Engine/Plugins/Compositing/Composite/Source/Composite/Composite.Build.cs
+++ b/Engine/Plugins/Compositing/Composite/Source/Composite/Composite.Build.cs
@@ -20,6 +20,7 @@ public class Composite : ModuleRules
 				"CompositeCore",
 				"OpenColorIO",
 				"MediaAssets",
+				"GameplayTags",
 			}
 		);
 
@@ -33,6 +34,7 @@ public class Composite : ModuleRules
 				"CameraCalibrationCore",
 				"MediaFrameworkUtilities",
 				"ConcertSyncCore",
+				"VPRoles",
 			}
 		);
 
diff --git a/Engine/Plugins/Compositing/Composite/Source/Composite/Private/CompositeActor.cpp b/Engine/Plugins/Compositing/Composite/Source/Composite/Private/CompositeActor.cpp
--- a/Engine/Plugins/Compositing/Composite/Source/Composite/Private/CompositeActor.cpp
+++ b/Engine/Plugins/Compositing/Composite/Source/Composite/Private/CompositeActor.cpp
@@ -12,6 +12,7 @@
 #include "Layers/CompositeLayerShadowReflection.h"
 #include "Layers/CompositeLayerPlate.h"
 #include "Layers/CompositeLayerMainRender.h"
+#include "VPRolesSubsystem.h"
 
 #include "Algo/Rotate.h"
 #include "Components/PrimitiveComponent.h"
@@ -325,6 +326,12 @@ void ACompositeActor::PostLoad()
 		}
 	}
 #endif // WITH_EDITORONLY_DATA
+
+	if (VPRole.IsValid())
+	{
+		// Note that if multiple actors have the same role, this could shift which actor is active on load.
+		SetActive(IsInVPRole());
+	}
 }
 
 void ACompositeActor::PostRegisterAllComponents()
@@ -534,6 +541,10 @@ void ACompositeActor::PostEditChangeChainProperty(struct FPropertyChangedChainEv
 			ViewProjectionComponent->SetTargetComponent(Camera);
 		}
 	}
+	else if (PropertyName == GET_MEMBER_NAME_CHECKED(ThisClass, VPRole))
+	{
+		SetVPRole(VPRole);
+	}
 }
 
 void ACompositeActor::PreEditUndo()
@@ -541,6 +552,7 @@ void ACompositeActor::PreEditUndo()
 	Super::PreEditUndo();
 
 	PreEditCompositeLayers = CompositeLayers;
+	PreEditVPRole = VPRole;
 }
 
 void ACompositeActor::PostEditUndo()
@@ -576,6 +588,12 @@ void ACompositeActor::PostEditUndo()
 		}
 	}
 
+	if (PreEditVPRole != VPRole)
+	{
+		// Manage activation if the vp role has changed
+		SetVPRole(VPRole);
+	}
+
 	Super::PostEditUndo();
 }
 
@@ -717,6 +735,18 @@ bool ACompositeActor::IsPlatePreviewActive() const
 }
 #endif
 
+FGameplayTag ACompositeActor::GetVPRole() const
+{
+	return VPRole;
+}
+
+void ACompositeActor::SetVPRole(FGameplayTag Value)
+{
+	VPRole = Value;
+
+	SetActive(IsInVPRole());
+}
+
 template<>
 UCompositeSceneCapture2DComponent* ACompositeActor::FindOrCreateSceneCapture(const UCompositeLayerBase* InLayer, int32 InIndex, FName InBaseName)
 {
@@ -865,11 +895,13 @@ bool ACompositeActor::IsRendering() const
 	return IsEnabled() && IsActive();
 }
 
-void ACompositeActor::PostJoinConcertSession()
+bool ACompositeActor::IsInVPRole() const
 {
-	SetActive(false);
+	// We are in a valid camera role if the user has not assigned a role or the current VPSettings role matches the
+	// assigned role.
+	UVirtualProductionRolesSubsystem* VPRolesSubsystem = GEngine->GetEngineSubsystem<UVirtualProductionRolesSubsystem>();
 
-	OnPostJoinMultiUserSession.Broadcast();
+	return !VPRole.IsValid() || (VPRolesSubsystem && VPRolesSubsystem->HasActiveRole(VPRole));
 }
 
 void ACompositeActor::PropagateStateChange(bool bApply)
diff --git a/Engine/Plugins/Compositing/Composite/Source/Composite/Private/CompositeWorldSubsystem.cpp b/Engine/Plugins/Compositing/Composite/Source/Composite/Private/CompositeWorldSubsystem.cpp
new file mode 100644
--- /dev/null
+++ b/Engine/Plugins/Compositing/Composite/Source/Composite/Private/CompositeWorldSubsystem.cpp
@@ -0,0 +1,62 @@
+// Copyright Epic Games, Inc. All Rights Reserved.
+
+#include "CompositeWorldSubsystem.h"
+#include "CompositeActor.h"
+#include "VPRolesSubsystem.h"
+
+#include "EngineUtils.h"
+#include "Engine/Engine.h"
+
+void UCompositeWorldSubsystem::PostInitialize()
+{
+	Super::PostInitialize();
+
+	UVirtualProductionRolesSubsystem* VPRolesSubsystem = GEngine->GetEngineSubsystem<UVirtualProductionRolesSubsystem>();
+	if (VPRolesSubsystem)
+	{
+		TWeakObjectPtr<UCompositeWorldSubsystem> WeakThis = this;
+
+		RolesChangedDelegateHandle = VPRolesSubsystem->OnRolesChanged().AddLambda([WeakThis](const TArray<FString>&)
+			{
+				TStrongObjectPtr<UCompositeWorldSubsystem> This = WeakThis.Pin();
+				if (!This.IsValid())
+				{
+					return;
+				}
+
+				const UWorld* World = This->GetWorld();
+				if (!World)
+				{
+					return;
+				}
+
+				for (TActorIterator<ACompositeActor> It(World); It; ++It)
+				{
+					ACompositeActor* CompositeActor = Cast<ACompositeActor>(*It);
+					if (!CompositeActor)
+					{
+						continue;
+					}
+
+					const FGameplayTag Role = CompositeActor->GetVPRole();
+					if (Role.IsValid())
+					{
+						// Update activation via setter when device role changes
+						CompositeActor->SetVPRole(Role);
+					}
+				}
+			}
+		);
+	}
+}
+
+void UCompositeWorldSubsystem::PreDeinitialize()
+{
+	Super::PreDeinitialize();
+
+	if (UVirtualProductionRolesSubsystem* VPRolesSubsystem = GEngine->GetEngineSubsystem<UVirtualProductionRolesSubsystem>())
+	{
+		VPRolesSubsystem->OnRolesChanged().Remove(RolesChangedDelegateHandle);
+	}
+}
+
diff --git a/Engine/Plugins/Compositing/Composite/Source/Composite/Private/CompositeWorldSubsystem.h b/Engine/Plugins/Compositing/Composite/Source/Composite/Private/CompositeWorldSubsystem.h
new file mode 100644
--- /dev/null
+++ b/Engine/Plugins/Compositing/Composite/Source/Composite/Private/CompositeWorldSubsystem.h
@@ -0,0 +1,27 @@
+// Copyright Epic Games, Inc. All Rights Reserved.
+
+#pragma once
+
+#include "Subsystems/WorldSubsystem.h"
+
+#include "CompositeWorldSubsystem.generated.h"
+
+/**
+ * World subsystem for composite actors.
+ */
+UCLASS(MinimalAPI)
+class UCompositeWorldSubsystem : public UWorldSubsystem
+{
+	GENERATED_BODY()
+
+public:
+	/** Called once all UWorldSubsystems have been initialized */
+	virtual void PostInitialize() override;
+
+	/** Called once all UWorldSubsystems are about to be deinitialized */
+	virtual void PreDeinitialize() override;
+
+private:
+	/** Event delegate for a device VP role change. */
+	FDelegateHandle RolesChangedDelegateHandle;
+};
diff --git a/Engine/Plugins/Compositing/Composite/Source/Composite/Public/CompositeActor.h b/Engine/Plugins/Compositing/Composite/Source/Composite/Public/CompositeActor.h
--- a/Engine/Plugins/Compositing/Composite/Source/Composite/Public/CompositeActor.h
+++ b/Engine/Plugins/Compositing/Composite/Source/Composite/Public/CompositeActor.h
@@ -4,6 +4,7 @@
 
 #include "CoreTypes.h"
 #include "GameFramework/Actor.h"
+#include "GameplayTagContainer.h"
 
 #include "Passes/CompositePassBase.h"
 #include "Layers/CompositeLayerBase.h"
@@ -18,8 +19,6 @@ class UCompositeShadowReflectionCatcherComponent;
 class UPrimitiveComponent;
 class UCompositeSceneCapture2DComponent;
 
-DECLARE_DYNAMIC_MULTICAST_DELEGATE(FOnCompositeActorPostJoinMultiUserSession);
-
 USTRUCT()
 struct FSceneCaptureComponentArray
 {
@@ -120,6 +119,14 @@ public:
 	UFUNCTION(BlueprintSetter)
 	UE_API void SetCompositeLayers(TArray<UCompositeLayerBase*> InLayers);
 
+	/** Get the vp role. */
+	UFUNCTION(BlueprintGetter)
+	UE_API FGameplayTag GetVPRole() const;
+	
+	/** Set the vp role. */
+	UFUNCTION(BlueprintSetter)
+	UE_API void SetVPRole(FGameplayTag Value);
+
 #if WITH_EDITOR
 	/**
 	 * Displays the plate texture preview in the level editor's 'picture-in-picture'
@@ -127,7 +134,6 @@ public:
 	 * @return true if the preview was started
 	 */
 	UE_API bool ShowPlatePreview(TObjectPtr<class UCompositePassColorKeyer> InRequestingKeyer);
-
 	/** Hides the plate texture preview */
 	UE_API void HidePlatePreview();
 
@@ -167,8 +173,8 @@ private:
 	/** Array of composite layers for merging or processing images. */
 	UPROPERTY(EditAnywhere, BlueprintGetter = GetCompositeLayers, BlueprintSetter = SetCompositeLayers, Instanced, Category = "Composite", meta = (DisplayName = "Layers"))
 	TArray<TObjectPtr<UCompositeLayerBase>> CompositeLayers;
-public:
 
+public:
 	/** Also provide the post-processing composite graph to SSR input. */
 	UPROPERTY(EditAnywhere, BlueprintReadWrite, Category = "Composite", AdvancedDisplay)
 	bool bEnableScreenSpaceReflections;
@@ -189,9 +195,14 @@ public:
 	UPROPERTY(EditAnywhere, BlueprintReadWrite, Category = "Composite", AdvancedDisplay)
 	ECompositeAllowedViewModes AllowedViewModes;
 
-	/** Delegate called after joining a multi-user session. Can be used to automatically re-active the actor locally. */
-	UPROPERTY(BlueprintAssignable, BlueprintCallable, BlueprintReadWrite, Category = "Composite")
-	FOnCompositeActorPostJoinMultiUserSession OnPostJoinMultiUserSession;
+private:
+	/**
+	 * Optional virtual production role to deactivate the actor when its role does not match the current device role.
+	 *
+	 * When unspecified, activation is unmodified (on load) or re-activated (on property change).
+	 */
+	UPROPERTY(EditAnywhere, BlueprintGetter = "GetVPRole", BlueprintSetter = "SetVPRole", Category = "Composite", AdvancedDisplay, meta = (DisplayName = "VP Role"))
+	FGameplayTag VPRole;
 
 public:
 	/** Convenience function used by layers to request a scene capture component. */
@@ -204,8 +215,8 @@ public:
 	/** Returns true if the composite is actively rendering. */
 	UE_API bool IsRendering() const;
 
-	/** Called after we join a multi-user session. */
-	UE_API void PostJoinConcertSession();
+	/** Check if the actor's role matches the session role. */
+	UE_API bool IsInVPRole() const;
 
 private:
 	/**
@@ -222,6 +233,9 @@ private:
 #if WITH_EDITORONLY_DATA
 	/** Used to update components after we remove layers. */
 	TArray<UCompositeLayerBase*> PreEditCompositeLayers;
+
+	/** Used to track changes to the VP role. */
+	FGameplayTag PreEditVPRole;
 #endif
 
 	/** Flag used to determine in which cases we should ignore PostRegister / Unregister calls. */
diff --git a/Engine/Plugins/Compositing/Composite/Source/CompositeEditor/CompositeEditor.Build.cs b/Engine/Plugins/Compositing/Composite/Source/CompositeEditor/CompositeEditor.Build.cs
--- a/Engine/Plugins/Compositing/Composite/Source/CompositeEditor/CompositeEditor.Build.cs
+++ b/Engine/Plugins/Compositing/Composite/Source/CompositeEditor/CompositeEditor.Build.cs
@@ -23,7 +23,6 @@ public class CompositeEditor : ModuleRules
 				"AppFramework",
 				"ColorGradingEditor",
 				"CompositeCore",
-				"ConcertSyncClient",
 				"ContentBrowserData",
 				"CoreUObject",
 				"DetailCustomizations",
diff --git a/Engine/Plugins/Compositing/Composite/Source/CompositeEditor/Private/CompositeEditorModule.cpp b/Engine/Plugins/Compositing/Composite/Source/CompositeEditor/Private/CompositeEditorModule.cpp
--- a/Engine/Plugins/Compositing/Composite/Source/CompositeEditor/Private/CompositeEditorModule.cpp
+++ b/Engine/Plugins/Compositing/Composite/Source/CompositeEditor/Private/CompositeEditorModule.cpp
@@ -8,9 +8,6 @@
 #include "EngineUtils.h"
 #include "HoldoutCompositeComponent.h"
 #include "Framework/Notifications/NotificationManager.h"
-#include "IConcertClientWorkspace.h"
-#include "IConcertSyncClient.h"
-#include "IConcertSyncClientModule.h"
 #include "LevelEditor.h"
 #include "Modules/ModuleManager.h"
 #include "Widgets/Notifications/SNotificationList.h"
@@ -70,18 +67,6 @@ void FCompositeEditorModule::StartupModule()
 
 void FCompositeEditorModule::ShutdownModule()
 {
-	if (ConcertSyncClient.IsValid())
-	{
-		ConcertSyncClient->OnWorkspaceStartup().RemoveAll(this);
-		ConcertSyncClient->OnWorkspaceShutdown().RemoveAll(this);
-		
-		TSharedPtr<IConcertClientWorkspace> WorkspacePtr = Workspace.Pin();
-		if (WorkspacePtr.IsValid())
-		{
-			WorkspacePtr->OnFinalizeWorkspaceSyncCompleted().RemoveAll(this);
-		}
-	}
-
 	UnregisterCustomizations();
 	UnregisterTabSpawners();
 
@@ -104,13 +89,6 @@ void FCompositeEditorModule::OnPostEngineInit()
 	{
 		GEditor->OnLevelActorAdded().AddRaw(this, &FCompositeEditorModule::OnLevelActorAdded);
 	}
-
-	ConcertSyncClient = IConcertSyncClientModule::Get().GetClient(TEXT("MultiUser"));
-	if (ConcertSyncClient.IsValid())
-	{
-		ConcertSyncClient->OnWorkspaceStartup().AddRaw(this, &FCompositeEditorModule::HandleWorkspaceStartup);
-		ConcertSyncClient->OnWorkspaceShutdown().AddRaw(this, &FCompositeEditorModule::HandleWorkspaceShutdown);
-	}
 }
 
 void FCompositeEditorModule::OnLevelActorAdded(AActor* InActor)
@@ -145,51 +123,6 @@ void FCompositeEditorModule::OnLevelActorAdded(AActor* InActor)
 	}
 }
 
-void FCompositeEditorModule::HandleWorkspaceStartup(const TSharedPtr<IConcertClientWorkspace>& NewWorkspace)
-{
-	Workspace = NewWorkspace;
-
-	if (NewWorkspace.IsValid())
-	{
-		NewWorkspace->OnFinalizeWorkspaceSyncCompleted().AddRaw(this, &FCompositeEditorModule::HandleFinalizeWorkspaceSyncCompleted);
-	}
-}
-
-void FCompositeEditorModule::HandleWorkspaceShutdown(const TSharedPtr<IConcertClientWorkspace>& WorkspaceShuttingDown)
-{
-	if (WorkspaceShuttingDown == Workspace)
-	{
-		TSharedPtr<IConcertClientWorkspace> WorkspacePtr = Workspace.Pin();
-		if (WorkspacePtr.IsValid())
-		{
-			WorkspacePtr->OnFinalizeWorkspaceSyncCompleted().RemoveAll(this);
-		}
-
-		Workspace.Reset();
-	}
-}
-
-void FCompositeEditorModule::HandleFinalizeWorkspaceSyncCompleted()
-{
-	UWorld* World = nullptr;
-	if (GEditor)
-	{
-		World = GEditor->GetEditorWorldContext().World();
-	}
-
-	if (IsValid(World))
-	{
-		for (TActorIterator<ACompositeActor> It(World); It; ++It)
-		{
-			ACompositeActor* CompositeActor = *It;
-			if (IsValid(CompositeActor))
-			{
-				CompositeActor->PostJoinConcertSession();
-			}
-		}
-	}
-}
-
 void FCompositeEditorModule::TriggerHoldoutCompositeWarning(const UHoldoutCompositeComponent* InComponent)
 {
 	static bool bWarnOnce = true;
diff --git a/Engine/Plugins/Compositing/Composite/Source/CompositeEditor/Private/UI/SCompositePanelLayerTree.cpp b/Engine/Plugins/Compositing/Composite/Source/CompositeEditor/Private/UI/SCompositePanelLayerTree.cpp
--- a/Engine/Plugins/Compositing/Composite/Source/CompositeEditor/Private/UI/SCompositePanelLayerTree.cpp
+++ b/Engine/Plugins/Compositing/Composite/Source/CompositeEditor/Private/UI/SCompositePanelLayerTree.cpp
@@ -399,6 +399,7 @@ private:
 						.Style(FAppStyle::Get(), "RadioButton")
 						.Visibility(this, &SCompositePanelLayerTreeItemRow::IsItemActiveVisible)
 						.IsChecked(this, &SCompositePanelLayerTreeItemRow::IsItemActive)
+						.IsEnabled(this, &SCompositePanelLayerTreeItemRow::IsItemActiveEnabled)
 						.OnCheckStateChanged(this, &SCompositePanelLayerTreeItemRow::OnIsItemActiveChanged)
 				];
 		}
@@ -505,6 +506,17 @@ private:
 		return ECheckBoxState::Unchecked;
 	}
 
+	/** Gets the checkbox enabled state for the row item's Active flag */
+	bool IsItemActiveEnabled() const
+	{
+		if (TStrongObjectPtr<ACompositeActor> CompositeActor = GetCompositeActor())
+		{
+			return CompositeActor->IsInVPRole();
+		}
+
+		return true;
+	}
+
 	/** Raised when the row item's Active radiobox has changed state */
 	void OnIsItemActiveChanged(ECheckBoxState InCheckBoxState)
 	{
@@ -1680,9 +1692,12 @@ void SCompositePanelLayerTree::OnGlobalActiveCheckStateChanged(ECheckBoxState Ch
 			TStrongObjectPtr<ACompositeActor> CompositeActor = TreeItem->CompositeActor.Pin();
 			if (CompositeActor.Get() == LastActiveActor.Get())
 			{
-				CompositeActor->SetActive(true);
-				OnCompositeActorActivated(CompositeActor);
-				bActivated = true;
+				if (CompositeActor->IsInVPRole())
+				{
+					CompositeActor->SetActive(true);
+					OnCompositeActorActivated(CompositeActor);
+					bActivated = true;
+				}
 			}
 		}
 
@@ -1697,9 +1712,12 @@ void SCompositePanelLayerTree::OnGlobalActiveCheckStateChanged(ECheckBoxState Ch
 				}
 
 				TStrongObjectPtr<ACompositeActor> CompositeActor = TreeItem->CompositeActor.Pin();
-				CompositeActor->SetActive(true);
-				OnCompositeActorActivated(CompositeActor);
-				break;
+				if (CompositeActor->IsInVPRole())
+				{
+					CompositeActor->SetActive(true);
+					OnCompositeActorActivated(CompositeActor);
+					break;
+				}
 			}
 		}
 	}
diff --git a/Engine/Plugins/Compositing/Composite/Source/CompositeEditor/Public/CompositeEditorModule.h b/Engine/Plugins/Compositing/Composite/Source/CompositeEditor/Public/CompositeEditorModule.h
--- a/Engine/Plugins/Compositing/Composite/Source/CompositeEditor/Public/CompositeEditorModule.h
+++ b/Engine/Plugins/Compositing/Composite/Source/CompositeEditor/Public/CompositeEditorModule.h
@@ -7,8 +7,6 @@
 #include "Logging/LogMacros.h"
 #include "Modules/ModuleManager.h"
 
-class IConcertClientWorkspace;
-
 /** Log category for this module. */
 DECLARE_LOG_CATEGORY_EXTERN(LogCompositeEditor, Log, All);
 
@@ -39,22 +37,8 @@ private:
 	void RegisterCustomizations();
 	void UnregisterCustomizations();
 
-	/** Workspace startup callback. */
-	void HandleWorkspaceStartup(const TSharedPtr<IConcertClientWorkspace>& NewWorkspace);
-	
-	/** Workspace shutdown callback. */
-	void HandleWorkspaceShutdown(const TSharedPtr<IConcertClientWorkspace>& WorkspaceShuttingDown);
-
-	/** Workspace sync completed callback. */
-	void HandleFinalizeWorkspaceSyncCompleted();
-
 private:
-	/** Cached concert sync client. */
-	TSharedPtr<class IConcertSyncClient> ConcertSyncClient;
 
 	/** Tab manager change delegate. */
 	FDelegateHandle OnTabManagerChangedDelegateHandle;
-
-	/** Holds a weak pointer to the current workspace. */
-	TWeakPtr<IConcertClientWorkspace> Workspace;
 };
