Composite: Fixed missing alpha unpremult/opacification needed for a texture plate behind holdout masks.

* Alpha opacification was previously only present on the built-in custom render pass for composite meshes.
* Extra pass added in "Texture" plate mode, when composite meshes are present.
* This also CL exposes a simple alpha transform pass, which can be used to manually transform alpha if needed.

#Synced-CL 50789144

diff --git a/Engine/Plugins/Compositing/Composite/Shaders/Private/CompositeAlphaTransform.usf b/Engine/Plugins/Compositing/Composite/Shaders/Private/CompositeAlphaTransform.usf
new file mode 100644
--- /dev/null
+++ b/Engine/Plugins/Compositing/Composite/Shaders/Private/CompositeAlphaTransform.usf
@@ -0,0 +1,49 @@
+// Copyright Epic Games, Inc. All Rights Reserved.
+
+#include "/Engine/Private/Common.ush"
+
+#include "CompositeCommon.ush"
+
+Texture2D Texture;
+SamplerState Sampler;
+uint PremultOp;
+uint bInvertSourceAlpha;
+uint OverrideOutputAlpha;
+
+#define PremultOp_None 0u
+#define PremultOp_Premultiply 1u
+#define PremultOp_Unpremultiply 2u
+
+#define OverrideOutput_None 0u
+#define OverrideOutput_Zero 1u
+#define OverrideOutput_One 2u
+
+void MainPS(
+    noperspective float4 InUVAndScreenPos : TEXCOORD0,
+    out float4 OutColor : SV_Target0
+    )
+{
+    float4 Color = Texture.SampleLevel(Sampler, InUVAndScreenPos.xy, 0);
+
+	Color.a = select(bInvertSourceAlpha > 0, saturate(1.0 - Color.a), Color.a);
+	
+	if(PremultOp == PremultOp_Premultiply)
+	{
+		Color.rgb *= Color.a;
+	}
+	else if(PremultOp == PremultOp_Unpremultiply)
+	{
+		Color.rgb *= SafeRcp(Color.a);
+	}
+
+	if(OverrideOutputAlpha == OverrideOutput_Zero)
+	{
+		Color.a = 0.0f;
+	}
+	else if(OverrideOutputAlpha == OverrideOutput_One)
+	{
+		Color.a = 1.0f;
+	}
+
+	OutColor = Color;
+}
diff --git a/Engine/Plugins/Compositing/Composite/Source/Composite/Private/Layers/CompositeLayerPlate.cpp b/Engine/Plugins/Compositing/Composite/Source/Composite/Private/Layers/CompositeLayerPlate.cpp
--- a/Engine/Plugins/Compositing/Composite/Source/Composite/Private/Layers/CompositeLayerPlate.cpp
+++ b/Engine/Plugins/Compositing/Composite/Source/Composite/Private/Layers/CompositeLayerPlate.cpp
@@ -16,6 +16,7 @@
 #include "Materials/Material.h"
 #include "Materials/MaterialInstanceDynamic.h"
 #include "Passes/CompositePassDistortion.h"
+#include "Passes/CompositePassAlphaTransform.h"
 #include "Profile/IMediaProfileManager.h"
 #include "Profile/MediaProfile.h"
 #include "Profile/MediaProfilePlaybackManager.h"
@@ -486,7 +487,9 @@ UE::CompositeCore::ResourceId UCompositeLayerPlate::AddPreprocessingPasses(
 bool UCompositeLayerPlate::GetProxy(FCompositeTraversalContext& InContext, FSceneRenderingBulkObjectAllocator& InFrameAllocator, FCompositeCorePassProxy*& OutProxy) const
 {
 	using namespace UE::CompositeCore;
+	using namespace UE::Composite::Private;
 
+	// Default to engine black dummy texture
 	ResourceId TextureId = BUILT_IN_EMPTY_ID;
 
 	if (IsValid(Texture))
@@ -517,14 +520,46 @@ bool UCompositeLayerPlate::GetProxy(FCompositeTraversalContext& InContext, FScen
 		);
 	}
 
-	if (PlateMode == ECompositePlateMode::CompositeMesh && GetValidPrimitivesNum() > 0)
-	{
-		// Prioritize the built-in custom render pass when enabled
-		TextureId = BUILT_IN_CRP_ID;
-	}
+	const bool bActiveCompositeMeshes = GetValidPrimitivesNum() > 0;
 
 	FPassInputDecl PassInput;
-	PassInput.Set<FPassExternalResourceDesc>({ TextureId });
+
+	if (bActiveCompositeMeshes)
+	{
+		switch (PlateMode)
+		{
+		case ECompositePlateMode::Texture:
+			{
+				// For accurate blending through the holdout mask, the built-in custom render pass
+				// for composite meshes already unpremultiplies the alpha and makes it opaque.
+				// 
+				// We replicate this behavior here automatically with an additional alpha transform pass
+				// when composite meshes are in use, and the PlateMode is set to "Texture".
+				FPassInputDecl AlphaPassInputDecl;
+				AlphaPassInputDecl.Set<FPassExternalResourceDesc>({ TextureId });
+
+				FAlphaTransformPassProxy* AlphaPassProxy = InFrameAllocator.Create<FAlphaTransformPassProxy>(FPassInputDeclArray{ AlphaPassInputDecl });
+				AlphaPassProxy->PremultOp = ECompositeAlphaPremultiplication::Unpremultiply;
+				AlphaPassProxy->OverrideOutputAlpha = ECompositeAlphaOverride::One;
+
+				PassInput.Set<const FCompositeCorePassProxy*>(AlphaPassProxy);
+			}
+			break;
+		
+		case ECompositePlateMode::CompositeMesh:
+			// Prioritize the built-in custom render pass when enabled over TextureId
+			PassInput.Set<FPassExternalResourceDesc>({ BUILT_IN_CRP_ID });
+			break;
+		
+		default:
+			checkNoEntry();
+			break;
+		}
+	}
+	else
+	{
+		PassInput.Set<FPassExternalResourceDesc>({ TextureId });
+	}
 
 	AddChildPasses(PassInput, InContext, InFrameAllocator, LayerPasses);
 
diff --git a/Engine/Plugins/Compositing/Composite/Source/Composite/Private/Passes/CompositePassAlphaTransform.cpp b/Engine/Plugins/Compositing/Composite/Source/Composite/Private/Passes/CompositePassAlphaTransform.cpp
new file mode 100644
--- /dev/null
+++ b/Engine/Plugins/Compositing/Composite/Source/Composite/Private/Passes/CompositePassAlphaTransform.cpp
@@ -0,0 +1,100 @@
+// Copyright Epic Games, Inc. All Rights Reserved.
+
+#include "Passes/CompositePassAlphaTransform.h"
+
+#include "Engine/Texture.h"
+#include "RenderGraphBuilder.h"
+#include "RenderGraphUtils.h"
+#include "SceneView.h"
+#include "ScreenPass.h"
+
+DECLARE_GPU_STAT_NAMED(FCompositeAlphaTransform, TEXT("Composite.AlphaTransform"));
+
+class FCompositePassAlphaTransformShader : public FGlobalShader
+{
+	DECLARE_GLOBAL_SHADER(FCompositePassAlphaTransformShader);
+	SHADER_USE_PARAMETER_STRUCT(FCompositePassAlphaTransformShader, FGlobalShader);
+
+	BEGIN_SHADER_PARAMETER_STRUCT(FParameters, )
+		SHADER_PARAMETER_RDG_TEXTURE(Texture2D, Texture)
+		SHADER_PARAMETER_SAMPLER(SamplerState, Sampler)
+		SHADER_PARAMETER(uint32, PremultOp)
+		SHADER_PARAMETER(uint32, bInvertSourceAlpha)
+		SHADER_PARAMETER(uint32, OverrideOutputAlpha)
+		RENDER_TARGET_BINDING_SLOTS()
+	END_SHADER_PARAMETER_STRUCT()
+};
+IMPLEMENT_GLOBAL_SHADER(FCompositePassAlphaTransformShader, "/Plugin/Composite/Private/CompositeAlphaTransform.usf", "MainPS", SF_Pixel);
+
+
+namespace UE
+{
+	namespace Composite
+	{
+		namespace Private
+		{
+			FPassTexture FAlphaTransformPassProxy::Add(FRDGBuilder& GraphBuilder, const FSceneView& InView, const FPassInputArray& Inputs, const FPassContext& PassContext) const
+			{
+				RDG_EVENT_SCOPE_STAT(GraphBuilder, FCompositeAlphaTransform, "Composite.AlphaTransform");
+
+				check(ValidateInputs(Inputs));
+
+				const UE::CompositeCore::FResourceMetadata& Metadata = Inputs[0].Metadata;
+				const FScreenPassTexture& InputTexture = Inputs[0].Texture;
+				FScreenPassRenderTarget Output = Inputs.OverrideOutput;
+
+				if (!Output.IsValid())
+				{
+					Output = CreateOutputRenderTarget(GraphBuilder, InView, PassContext.OutputViewRect, InputTexture.Texture->Desc, TEXT("AlphaTransformCompositePass"));
+				}
+
+				FGlobalShaderMap* GlobalShaderMap = GetGlobalShaderMap(InView.GetFeatureLevel());
+				FScreenPassTextureViewport Viewport = FScreenPassTextureViewport(Output);
+
+				FCompositePassAlphaTransformShader::FParameters* Parameters = GraphBuilder.AllocParameters<FCompositePassAlphaTransformShader::FParameters>();
+				Parameters->Texture = InputTexture.Texture;
+				Parameters->Sampler = TStaticSamplerState<SF_Bilinear>::GetRHI();
+				Parameters->PremultOp = static_cast<uint32>(PremultOp);
+				Parameters->bInvertSourceAlpha = static_cast<uint32>(bInvertSourceAlpha);
+				Parameters->OverrideOutputAlpha = static_cast<uint32>(OverrideOutputAlpha);
+				Parameters->RenderTargets[0] = Output.GetRenderTargetBinding();
+
+				TShaderMapRef<FCompositePassAlphaTransformShader> PixelShader(GlobalShaderMap);
+				AddDrawScreenPass(
+					GraphBuilder,
+					RDG_EVENT_NAME("Composite.AlphaTransform (%dx%d)", Viewport.Extent.X, Viewport.Extent.Y),
+					InView,
+					Viewport,
+					Viewport,
+					PixelShader,
+					Parameters
+				);
+
+				return UE::CompositeCore::FPassTexture{ MoveTemp(Output), Metadata };
+			}
+		}
+	}
+}
+
+UCompositePassAlphaTransform::UCompositePassAlphaTransform(const FObjectInitializer& ObjectInitializer)
+	: Super(ObjectInitializer)
+	, PremultOp{ ECompositeAlphaPremultiplication::None }
+	, bInvertSourceAlpha{ false }
+	, OverrideOutputAlpha{ ECompositeAlphaOverride::None }
+{
+}
+
+UCompositePassAlphaTransform::~UCompositePassAlphaTransform() = default;
+
+bool UCompositePassAlphaTransform::GetProxy(const UE::CompositeCore::FPassInputDecl& InputDecl, FCompositeTraversalContext& InContext, FSceneRenderingBulkObjectAllocator& InFrameAllocator, FCompositeCorePassProxy*& OutProxy) const
+{
+	using namespace UE::Composite::Private;
+
+	FAlphaTransformPassProxy* Proxy = InFrameAllocator.Create<FAlphaTransformPassProxy>(FPassInputDeclArray{ InputDecl });
+	Proxy->PremultOp = PremultOp;
+	Proxy->bInvertSourceAlpha = bInvertSourceAlpha;
+	Proxy->OverrideOutputAlpha = OverrideOutputAlpha;
+
+	OutProxy = Proxy;
+	return true;
+}
diff --git a/Engine/Plugins/Compositing/Composite/Source/Composite/Public/Passes/CompositePassAlphaTransform.h b/Engine/Plugins/Compositing/Composite/Source/Composite/Public/Passes/CompositePassAlphaTransform.h
new file mode 100644
--- /dev/null
+++ b/Engine/Plugins/Compositing/Composite/Source/Composite/Public/Passes/CompositePassAlphaTransform.h
@@ -0,0 +1,84 @@
+// Copyright Epic Games, Inc. All Rights Reserved.
+
+#pragma once
+
+#include "CompositePassBase.h"
+#include "Passes/CompositeCorePassProxy.h"
+
+#include "CompositePassAlphaTransform.generated.h"
+
+#define UE_API COMPOSITE_API
+
+/** Alpha premultiplication operation type. */
+UENUM(BlueprintType)
+enum class ECompositeAlphaPremultiplication : uint8
+{
+	None,
+	Premultiply,
+	Unpremultiply,
+};
+
+UENUM(BlueprintType)
+enum class ECompositeAlphaOverride : uint8
+{
+	None,
+	Zero,
+	One,
+};
+
+/** Alpha transform pass. */
+UCLASS(MinimalAPI, BlueprintType, Blueprintable, EditInlineNew, CollapseCategories, meta=(DisplayName="Alpha Transform Pass"))
+class UCompositePassAlphaTransform : public UCompositePassBase
+{
+	GENERATED_BODY()
+
+public:
+	/** Constructor */
+	UE_API UCompositePassAlphaTransform(const FObjectInitializer& ObjectInitializer);
+	
+	/** Destructor */
+	UE_API ~UCompositePassAlphaTransform();
+
+	UE_API virtual bool GetProxy(const UE::CompositeCore::FPassInputDecl& InputDecl, FCompositeTraversalContext& InContext, FSceneRenderingBulkObjectAllocator& InFrameAllocator, FCompositeCorePassProxy*& OutProxy) const override;
+
+public:
+	/** Alpha premultiplication operation. */
+	UPROPERTY(BlueprintReadWrite, EditAnywhere, Category = "")
+	ECompositeAlphaPremultiplication PremultOp;
+
+	/** Invert the source alpha, for instance if it comes from a scene capture render. */
+	UPROPERTY(BlueprintReadWrite, EditAnywhere, Category = "", AdvancedDisplay)
+	bool bInvertSourceAlpha;
+
+	/** Output alpha is forced to 1.0, primarily used for correct blending with composite mesh holdouts. */
+	UPROPERTY(BlueprintReadWrite, EditAnywhere, Category = "", AdvancedDisplay)
+	ECompositeAlphaOverride OverrideOutputAlpha;
+};
+
+namespace UE
+{
+	namespace Composite
+	{
+		namespace Private
+		{
+			using namespace CompositeCore;
+
+			class FAlphaTransformPassProxy : public FCompositeCorePassProxy
+			{
+			public:
+				IMPLEMENT_COMPOSITE_PASS(FAlphaTransformPassProxy);
+
+				using FCompositeCorePassProxy::FCompositeCorePassProxy;
+
+				FPassTexture Add(FRDGBuilder& GraphBuilder, const FSceneView& InView, const FPassInputArray& Inputs, const FPassContext& PassContext) const override;
+
+				ECompositeAlphaPremultiplication PremultOp = ECompositeAlphaPremultiplication::None;
+				bool bInvertSourceAlpha = false;
+				ECompositeAlphaOverride OverrideOutputAlpha = ECompositeAlphaOverride::None;
+			};
+		}
+	}
+}
+
+#undef UE_API
+
