Composure: Restore composite mesh visibility when removed (minor fix following CL 50672850).

#Synced-CL 50671255

diff --git a/Engine/Plugins/Compositing/Composite/Source/Composite/Private/Layers/CompositeLayerPlate.cpp b/Engine/Plugins/Compositing/Composite/Source/Composite/Private/Layers/CompositeLayerPlate.cpp
--- a/Engine/Plugins/Compositing/Composite/Source/Composite/Private/Layers/CompositeLayerPlate.cpp
+++ b/Engine/Plugins/Compositing/Composite/Source/Composite/Private/Layers/CompositeLayerPlate.cpp
@@ -47,8 +47,8 @@ namespace UE
 	{
 		namespace Private
 		{
-			/** Convenience function to remove composite asset user data from tracked mesh actors (before their removal). */
-			static void RemoveCompositeAssetUserData(TConstArrayView<TSoftObjectPtr<AActor>> InPreEditMeshes, TConstArrayView<TSoftObjectPtr<AActor>> InPostEditMeshes)
+			/** Convenience function to handle state on removed composite meshes. */
+			static void HandleRemovedCompositeMeshes(TConstArrayView<TSoftObjectPtr<AActor>> InPreEditMeshes, TConstArrayView<TSoftObjectPtr<AActor>> InPostEditMeshes, bool bRestoreVisibility)
 			{
 				for (const TSoftObjectPtr<AActor>& PreEditMesh : InPreEditMeshes)
 				{
@@ -66,6 +66,12 @@ namespace UE
 							if (IsValid(PrimitiveComponent))
 							{
 								PrimitiveComponent->RemoveUserDataOfClass(UCompositeAssetUserData::StaticClass());
+
+								if (bRestoreVisibility)
+								{
+									constexpr bool bNewVisibility = true;
+									PrimitiveComponent->SetVisibility(true);
+								}
 							}
 						}
 					}
@@ -75,7 +81,8 @@ namespace UE
 			/** Convenience function to remove all composite asset user data from tracked mesh actors. */
 			static void RemoveCompositeAllAssetUserData(TConstArrayView<TSoftObjectPtr<AActor>> InPreEditMeshes)
 			{
-				RemoveCompositeAssetUserData(InPreEditMeshes, {});
+				constexpr bool bRestoreVisibility = false;
+				HandleRemovedCompositeMeshes(InPreEditMeshes, {}, bRestoreVisibility);
 			}
 		}
 	}
@@ -230,8 +237,10 @@ const TArray<TSoftObjectPtr<AActor>> UCompositeLayerPlate::GetCompositeMeshes()
 
 void UCompositeLayerPlate::SetCompositeMeshes(TArray<TSoftObjectPtr<AActor>> InCompositeMeshes)
 {
+	const bool bControlMeshVisibility = CVarCompositePlateControlMeshVisibility.GetValueOnGameThread();
+
 	// First, we untrack current meshes
-	UE::Composite::Private::RemoveCompositeAssetUserData(CompositeMeshes, InCompositeMeshes);
+	UE::Composite::Private::HandleRemovedCompositeMeshes(CompositeMeshes, InCompositeMeshes, bControlMeshVisibility);
 
 	CompositeMeshes = MoveTemp(InCompositeMeshes);
 	
@@ -332,7 +341,9 @@ void UCompositeLayerPlate::PostEditChangeProperty(FPropertyChangedEvent& Propert
 
 		if (bChangeWasRemoval)
 		{
-			UE::Composite::Private::RemoveCompositeAssetUserData(PreEditCompositeMeshes, CompositeMeshes);
+			const bool bControlMeshVisibility = CVarCompositePlateControlMeshVisibility.GetValueOnGameThread();
+
+			UE::Composite::Private::HandleRemovedCompositeMeshes(PreEditCompositeMeshes, CompositeMeshes, bControlMeshVisibility);
 		}
 		else
 		{
