Composure: Added support for only using despill on the keyer pass, as a new advanced setting.

* While we don't expect this to be the primary path for Ultimatte support, users may chose to apply an external key mask but despill the raw video feed in UE.

#Synced-CL 50934488

diff --git a/Engine/Plugins/Compositing/Composite/Shaders/Private/CompositeColorKeyer.usf b/Engine/Plugins/Compositing/Composite/Shaders/Private/CompositeColorKeyer.usf
--- a/Engine/Plugins/Compositing/Composite/Shaders/Private/CompositeColorKeyer.usf
+++ b/Engine/Plugins/Compositing/Composite/Shaders/Private/CompositeColorKeyer.usf
@@ -86,6 +86,7 @@ void MainCS(
 
 	float4 Result = Foreground;
 
+#if !USE_DESPILL_ONLY
 	float2 ScreenPos = UVToScreenPos(UV, Input_Extent);
 	float AspectRatio = Output_ViewportSize.y * Output_ViewportSizeInverse.x;
 	float VignetteMask = Vignette(ScreenPos, AspectRatio, DevignetteStrength);
@@ -108,6 +109,7 @@ void MainCS(
 	{
 		Result = BlueScreenKeyer(Result, Background, ChannelWeights[0], ChannelWeights[1]);
 	}
+#endif
 
 	BRANCH
 	if(DespillStrength > 0.0f)
@@ -127,6 +129,7 @@ void MainCS(
 		Result.rgb = Luv_to_RGB(Foreground_Luv);
 	}
 
+#if !USE_DESPILL_ONLY
 	if(bPreserveVignetteAfterKey > 0 && DevignetteStrength > ZeroEpsilon)
 	{
 		Result.rgb *= VignetteMask;
@@ -134,6 +137,7 @@ void MainCS(
 
 	// Apply final alpha threshold
 	Result.a = Linearstep(AlphaMin, AlphaMax, Result.a);
+#endif
 	
 	const uint2 PixelPos = DispatchThreadId + Output_ViewportMin;
 	RWOutputTexture[PixelPos] = Output(Result);
diff --git a/Engine/Plugins/Compositing/Composite/Source/Composite/Private/Passes/CompositePassColorKeyer.cpp b/Engine/Plugins/Compositing/Composite/Source/Composite/Private/Passes/CompositePassColorKeyer.cpp
--- a/Engine/Plugins/Compositing/Composite/Source/Composite/Private/Passes/CompositePassColorKeyer.cpp
+++ b/Engine/Plugins/Compositing/Composite/Source/Composite/Private/Passes/CompositePassColorKeyer.cpp
@@ -43,7 +43,8 @@ class FCompositePassColorKeyerShader : public FGlobalShader
 	SHADER_USE_PARAMETER_STRUCT(FCompositePassColorKeyerShader, FGlobalShader);
 
 	class FUseCleanPlate : SHADER_PERMUTATION_BOOL("USE_CLEAN_PLATE");
-	using FPermutationDomain = TShaderPermutationDomain<FUseCleanPlate>;
+	class FDespillOnly : SHADER_PERMUTATION_BOOL("USE_DESPILL_ONLY");
+	using FPermutationDomain = TShaderPermutationDomain<FUseCleanPlate, FDespillOnly>;
 
 	BEGIN_SHADER_PARAMETER_STRUCT(FParameters, )
 		SHADER_PARAMETER_RDG_TEXTURE(Texture2D, InputTexture)
@@ -178,6 +179,7 @@ namespace UE
 
 				FCompositePassColorKeyerShader::FPermutationDomain PermutationVector;
 				PermutationVector.Set<FCompositePassColorKeyerShader::FUseCleanPlate>(CleanPlateResource != nullptr);
+				PermutationVector.Set<FCompositePassColorKeyerShader::FDespillOnly>(bDespillOnly);
 
 				const FIntPoint OutputSize = Output.ViewRect.Size();
 				TShaderMapRef<FCompositePassColorKeyerShader> ComputeShader(GlobalShaderMap, PermutationVector);
@@ -205,6 +207,7 @@ namespace UE
 			ECompositeColorKeyerVisualization Visualization;
 			bool bPreserveVignetteAfterKey;
 			ECompositeDenoiseMethod DenoiseMethod;
+			bool bDespillOnly;
 			bool bInvertAlpha;
 			TWeakObjectPtr<UTexture> CleanPlateWeakPtr;
 		};
@@ -225,6 +228,7 @@ UCompositePassColorKeyer::UCompositePassColorKeyer(const FObjectInitializer& Obj
 	, bPreserveVignetteAfterKey(true)
 	, DenoiseMethod(ECompositeDenoiseMethod::None)
 	, Visualization(ECompositeColorKeyerVisualization::Key)
+	, bDespillOnly(false)
 	, bInvertAlpha(false)
 { 
 }
@@ -242,6 +246,7 @@ bool UCompositePassColorKeyer::GetProxy(const UE::CompositeCore::FPassInputDecl&
 	Proxy->ScreenType = ScreenType;
 	Proxy->bPreserveVignetteAfterKey = bPreserveVignetteAfterKey;
 	Proxy->DenoiseMethod = DenoiseMethod;
+	Proxy->bDespillOnly = bDespillOnly;
 	Proxy->bInvertAlpha = bInvertAlpha;
 	Proxy->Visualization = Visualization;
 
diff --git a/Engine/Plugins/Compositing/Composite/Source/Composite/Public/Passes/CompositePassColorKeyer.h b/Engine/Plugins/Compositing/Composite/Source/Composite/Public/Passes/CompositePassColorKeyer.h
--- a/Engine/Plugins/Compositing/Composite/Source/Composite/Public/Passes/CompositePassColorKeyer.h
+++ b/Engine/Plugins/Compositing/Composite/Source/Composite/Public/Passes/CompositePassColorKeyer.h
@@ -130,6 +130,10 @@ public:
 	UPROPERTY(EditAnywhere, BlueprintReadWrite, Category = "")
 	ECompositeColorKeyerVisualization Visualization;
 
+	/** Only apply the despill algorithm controlled by DespillStrength, alpha is unmodified. */
+	UPROPERTY(EditAnywhere, BlueprintReadWrite, Category = "", AdvancedDisplay)
+	bool bDespillOnly;
+
 	/** Invert the alpha key. */
 	UPROPERTY(EditAnywhere, BlueprintReadWrite, Category = "", AdvancedDisplay)
 	bool bInvertAlpha;
