Composure: Reduce default render (target) resolution to improve performance with shadow/reflection layer.

* Before CLs 49938137 and 49940431, scene captures we incorrectly rendering at main view resolution despite their render targets being fixed. In the editor viewport, this often meant a resolution smaller than 1920x1080. As we now respect 1080p, performance appears worse. As a fix, we directly change the default scene capture resolution to 960x540, which tends to be sufficient for reflections & shadows.

#Synced-CL 50934488

diff --git a/Engine/Plugins/Compositing/Composite/Source/Composite/Private/CompositeActor.cpp b/Engine/Plugins/Compositing/Composite/Source/Composite/Private/CompositeActor.cpp
--- a/Engine/Plugins/Compositing/Composite/Source/Composite/Private/CompositeActor.cpp
+++ b/Engine/Plugins/Compositing/Composite/Source/Composite/Private/CompositeActor.cpp
@@ -252,10 +252,7 @@ ACompositeActor::ACompositeActor(const FObjectInitializer& ObjectInitializer)
 	, ViewProjectionComponent{}
 	, bIsActive{true}
 	, bIsEnabled{true}
-// TODO: Can we find the current viewport render resolution instead?
-// View rects are only accessible from the render thread, but maybe there is an alternative.
-// This isn't needed by custom render passes, only regular scene capture render targets...
-	, RenderResolution{1920, 1080}
+	, RenderResolution{ FCompositeRenderTargetPool::DefaultSize }
 	, bEnableScreenSpaceReflections{true}
 	, bOverridesViewUserFlags{true}
 	, ViewUserFlags{1}
diff --git a/Engine/Plugins/Compositing/Composite/Source/Composite/Private/CompositeRenderTargetPool.cpp b/Engine/Plugins/Compositing/Composite/Source/Composite/Private/CompositeRenderTargetPool.cpp
--- a/Engine/Plugins/Compositing/Composite/Source/Composite/Private/CompositeRenderTargetPool.cpp
+++ b/Engine/Plugins/Compositing/Composite/Source/Composite/Private/CompositeRenderTargetPool.cpp
@@ -14,7 +14,7 @@ static TAutoConsoleVariable<int32> CVarCompositePoolFramesUntilFlush(
 	TEXT("This may leave render resources pooled, but never reclaimed. Auto-Flushing returns those resources. ")
 	TEXT("For values greater than zero, the pooling system will wait that number of frames before a target is considered 'stale'."));
 
-const FIntPoint FCompositeRenderTargetPool::DefaultSize = FIntPoint(1920, 1080);
+const FIntPoint FCompositeRenderTargetPool::DefaultSize = FIntPoint(960, 540);
 
 namespace UE
 {
diff --git a/Engine/Plugins/Compositing/Composite/Source/Composite/Private/CompositeRenderTargetPool.h b/Engine/Plugins/Compositing/Composite/Source/Composite/Private/CompositeRenderTargetPool.h
--- a/Engine/Plugins/Compositing/Composite/Source/Composite/Private/CompositeRenderTargetPool.h
+++ b/Engine/Plugins/Compositing/Composite/Source/Composite/Private/CompositeRenderTargetPool.h
@@ -26,7 +26,7 @@ class FCompositeRenderTargetPool final : public FGCObject
 #endif
 {
 public:
-	/** Default render target size/resolution of 1080p */
+	/** Default render target size/resolution. */
 	static const FIntPoint DefaultSize;
 
 	/** Static singleton pool getter. */
diff --git a/Engine/Plugins/Compositing/Composite/Source/Composite/Public/CompositeActor.h b/Engine/Plugins/Compositing/Composite/Source/Composite/Public/CompositeActor.h
--- a/Engine/Plugins/Compositing/Composite/Source/Composite/Public/CompositeActor.h
+++ b/Engine/Plugins/Compositing/Composite/Source/Composite/Public/CompositeActor.h
@@ -161,7 +161,7 @@ private:
 	bool bIsEnabled;
 
 public:
-	/** Composite render (output) resolution. Currently only used for transient scene capture render targets. */
+	/** Composite render target resolution, used by managed scene captures. */
 	UPROPERTY(EditAnywhere, BlueprintReadWrite, Interp, Category = "Composite")
 	FIntPoint RenderResolution;
 
diff --git a/Engine/Plugins/Compositing/Composite/Source/CompositeEditor/Private/UI/SCompositeEditorPanel.cpp b/Engine/Plugins/Compositing/Composite/Source/CompositeEditor/Private/UI/SCompositeEditorPanel.cpp
--- a/Engine/Plugins/Compositing/Composite/Source/CompositeEditor/Private/UI/SCompositeEditorPanel.cpp
+++ b/Engine/Plugins/Compositing/Composite/Source/CompositeEditor/Private/UI/SCompositeEditorPanel.cpp
@@ -108,10 +108,10 @@ public:
 		InHeaderRow.OverrideResetToDefault(FResetToDefaultOverride::Create(
 			FIsResetToDefaultVisible::CreateLambda([this](TSharedPtr<IPropertyHandle>)
 			{
-				return ResolutionComboBox->GetSelectedItem() != NamedResolutions[1];
+				return ResolutionComboBox->GetSelectedItem() != NamedResolutions[0];
 			}), FResetToDefaultHandler::CreateLambda([this](TSharedPtr<IPropertyHandle>)
 			{
-				ResolutionComboBox->SetSelectedItem(NamedResolutions[1]);
+				ResolutionComboBox->SetSelectedItem(NamedResolutions[0]);
 			})));
 	}
 
@@ -307,6 +307,7 @@ private:
 
 TArray<TSharedPtr<FResolutionTypeCustomization::FNamedResolution>> FResolutionTypeCustomization::NamedResolutions =
 {
+	MakeShared<FNamedResolution>(LOCTEXT("540pResolutionName", "540p (qHD) - 960x540"), FIntPoint(960, 540), LOCTEXT("Resolution540pToolTip", "Quarter High Definition")),
 	MakeShared<FNamedResolution>(LOCTEXT("720pResolutionName", "720p (HD) - 1280x720"), FIntPoint(1280, 720), LOCTEXT("Resolution720pToolTip", "High Definition")),
 	MakeShared<FNamedResolution>(LOCTEXT("1080pResolutionName", "1080p (FHD) - 1920x1080"), FIntPoint(1920, 1080), LOCTEXT("Resolution1080pToolTip", "Full high definition")),
 	MakeShared<FNamedResolution>(LOCTEXT("1440pResolutionName", "1440p (QHD) - 2560x1440"), FIntPoint(2560, 1440), LOCTEXT("Resolution1440pToolTip", "Quad high definition")),
