Composure: Minor refactor to inherit from composite core world subsystem.

We also remove potentially slow actor world iterator calls in favor of register actory queries on the subsystem.

#Synced-CL 50960501

diff --git a/Engine/Plugins/Compositing/Composite/Source/Composite/Private/CompositeActor.cpp b/Engine/Plugins/Compositing/Composite/Source/Composite/Private/CompositeActor.cpp
--- a/Engine/Plugins/Compositing/Composite/Source/Composite/Private/CompositeActor.cpp
+++ b/Engine/Plugins/Compositing/Composite/Source/Composite/Private/CompositeActor.cpp
@@ -2,7 +2,7 @@
 
 #include "CompositeActor.h"
 
-#include "CompositeCoreSubsystem.h"
+#include "CompositeWorldSubsystem.h"
 #include "CompositeModule.h"
 #include "CompositeCoreSettings.h"
 #include "CompositeRenderTargetPool.h"
@@ -331,6 +331,16 @@ void ACompositeActor::PostLoad()
 	}
 }
 
+void ACompositeActor::PreRegisterAllComponents()
+{
+	Super::PreRegisterAllComponents();
+
+	if (UCompositeWorldSubsystem* Subsystem = UWorld::GetSubsystem<UCompositeWorldSubsystem>(GetWorld()))
+	{
+		Subsystem->RegisterActor(this);
+	}
+}
+
 void ACompositeActor::PostRegisterAllComponents()
 {
 	using namespace UE::Composite::Private;
@@ -368,6 +378,11 @@ void ACompositeActor::UnregisterAllComponents(bool bForReregister)
 void ACompositeActor::PostUnregisterAllComponents()
 {
 	Super::PostUnregisterAllComponents();
+
+	if (UCompositeWorldSubsystem* Subsystem = UWorld::GetSubsystem<UCompositeWorldSubsystem>(GetWorld()))
+	{
+		Subsystem->UnregisterActor(this);
+	}
 }
 
 void ACompositeActor::Tick(float DeltaTime)
@@ -381,8 +396,7 @@ void ACompositeActor::Tick(float DeltaTime)
 		return;
 	}
 
-	UCompositeCoreSubsystem* Subsystem = UWorld::GetSubsystem<UCompositeCoreSubsystem>(GetWorld());
-	if (IsValid(Subsystem))
+	if (UCompositeWorldSubsystem* Subsystem = UWorld::GetSubsystem<UCompositeWorldSubsystem>(GetWorld()))
 	{
 		Subsystem->SetRenderWork(UE::Composite::Private::GatherFrameRenderWork(*this, SceneCapturesPerLayer));
 
@@ -568,19 +582,20 @@ void ACompositeActor::PostEditUndo()
 
 	PreEditCompositeLayers.Reset();
 
-	UWorld* World = GetWorld();
-
 	// Deactivate the current actor if any other actor is already active
 	// This can a occur when a previously active actor deletion is undone.
-	if (IsValid(World) && IsActive())
+	if (IsActive())
 	{
-		for (TActorIterator<ACompositeActor> It(World); It; ++It)
+		if (UCompositeWorldSubsystem* Subsystem = UWorld::GetSubsystem<UCompositeWorldSubsystem>(GetWorld()))
 		{
-			ACompositeActor* OtherActor = *It;
-			if ((OtherActor != this) && IsValid(OtherActor) && OtherActor->IsActive())
+			for (const TWeakObjectPtr<ACompositeActor> Actor : Subsystem->GetActors())
 			{
-				SetActive(false);
-				break;
+				ACompositeActor* OtherActor = Actor.Get();
+				if ((OtherActor != this) && IsValid(OtherActor) && OtherActor->IsActive())
+				{
+					SetActive(false);
+					break;
+				}
 			}
 		}
 	}
@@ -639,16 +654,18 @@ void ACompositeActor::SetActive(bool bInActive)
 	bIsActive = bInActive;
 
 	// Deactivate all other composite actors first
-	const UWorld* World = GetWorld();
-
-	if (bIsActive && IsValid(World))
+	if (bIsActive)
 	{
-		for (TActorIterator<ACompositeActor> It(World); It; ++It)
+		if (UCompositeWorldSubsystem* Subsystem = UWorld::GetSubsystem<UCompositeWorldSubsystem>(GetWorld()))
 		{
-			if (*It != this)
+			for (const TWeakObjectPtr<ACompositeActor> WeakActor : Subsystem->GetActors())
 			{
-				It->bIsActive = false;
-				It->PropagateStateChange(false);
+				ACompositeActor* Actor = WeakActor.Get();
+				if (Actor && Actor != this)
+				{
+					Actor->bIsActive = false;
+					Actor->PropagateStateChange(false);
+				}
 			}
 		}
 	}
@@ -914,8 +931,7 @@ void ACompositeActor::PropagateStateChange(bool bApply)
 	// Remove state case
 	if (!bApply)
 	{
-		UCompositeCoreSubsystem* Subsystem = UWorld::GetSubsystem<UCompositeCoreSubsystem>(GetWorld());
-		if (IsValid(Subsystem))
+		if (UCompositeWorldSubsystem* Subsystem = UWorld::GetSubsystem<UCompositeWorldSubsystem>(GetWorld()))
 		{
 			Subsystem->ResetRenderWork();
 			Subsystem->ResetBuiltInRenderPassOptions();
diff --git a/Engine/Plugins/Compositing/Composite/Source/Composite/Private/CompositeWorldSubsystem.cpp b/Engine/Plugins/Compositing/Composite/Source/Composite/Private/CompositeWorldSubsystem.cpp
--- a/Engine/Plugins/Compositing/Composite/Source/Composite/Private/CompositeWorldSubsystem.cpp
+++ b/Engine/Plugins/Compositing/Composite/Source/Composite/Private/CompositeWorldSubsystem.cpp
@@ -30,9 +30,9 @@ void UCompositeWorldSubsystem::PostInitialize()
 					return;
 				}
 
-				for (TActorIterator<ACompositeActor> It(World); It; ++It)
+				for (const TWeakObjectPtr<ACompositeActor>& Actor : This->CompositeActors)
 				{
-					ACompositeActor* CompositeActor = Cast<ACompositeActor>(*It);
+					ACompositeActor* CompositeActor = Actor.Get();
 					if (!CompositeActor)
 					{
 						continue;
@@ -60,3 +60,21 @@ void UCompositeWorldSubsystem::PreDeinitialize()
 	}
 }
 
+void UCompositeWorldSubsystem::RegisterActor(TWeakObjectPtr<ACompositeActor> InActor)
+{
+	if (InActor.IsValid() && InActor->GetWorld() == GetWorld())
+	{
+		CompositeActors.Add(InActor);
+	}
+}
+
+void UCompositeWorldSubsystem::UnregisterActor(TWeakObjectPtr<ACompositeActor> InActor)
+{
+	CompositeActors.Remove(InActor);
+}
+
+const TSet<TWeakObjectPtr<ACompositeActor>>& UCompositeWorldSubsystem::GetActors() const
+{
+	return CompositeActors;
+}
+
diff --git a/Engine/Plugins/Compositing/Composite/Source/Composite/Private/CompositeWorldSubsystem.h b/Engine/Plugins/Compositing/Composite/Source/Composite/Private/CompositeWorldSubsystem.h
--- a/Engine/Plugins/Compositing/Composite/Source/Composite/Private/CompositeWorldSubsystem.h
+++ b/Engine/Plugins/Compositing/Composite/Source/Composite/Private/CompositeWorldSubsystem.h
@@ -2,15 +2,17 @@
 
 #pragma once
 
-#include "Subsystems/WorldSubsystem.h"
+#include "CompositeCoreSubsystem.h"
 
 #include "CompositeWorldSubsystem.generated.h"
 
+class ACompositeActor;
+
 /**
- * World subsystem for composite actors.
+ * Subsystem for composite actors, deriving from the CompositeCore subsystem.
  */
 UCLASS(MinimalAPI)
-class UCompositeWorldSubsystem : public UWorldSubsystem
+class UCompositeWorldSubsystem : public UCompositeCoreSubsystem
 {
 	GENERATED_BODY()
 
@@ -21,7 +23,19 @@ public:
 	/** Called once all UWorldSubsystems are about to be deinitialized */
 	virtual void PreDeinitialize() override;
 
+	/** Register composite actor. */
+	void RegisterActor(TWeakObjectPtr<ACompositeActor> InActor);
+
+	/** Unregister composite actor. */
+	void UnregisterActor(TWeakObjectPtr<ACompositeActor> InActor);
+
+	/** Get all registered composite actors. */
+	const TSet<TWeakObjectPtr<ACompositeActor>>& GetActors() const;
+
 private:
+	/** List of registered actors. */
+	TSet<TWeakObjectPtr<ACompositeActor>> CompositeActors;
+
 	/** Event delegate for a device VP role change. */
 	FDelegateHandle RolesChangedDelegateHandle;
 };
diff --git a/Engine/Plugins/Compositing/Composite/Source/Composite/Private/Layers/CompositeLayerPlate.cpp b/Engine/Plugins/Compositing/Composite/Source/Composite/Private/Layers/CompositeLayerPlate.cpp
--- a/Engine/Plugins/Compositing/Composite/Source/Composite/Private/Layers/CompositeLayerPlate.cpp
+++ b/Engine/Plugins/Compositing/Composite/Source/Composite/Private/Layers/CompositeLayerPlate.cpp
@@ -5,10 +5,10 @@
 #include "Components/StaticMeshComponent.h"
 #include "CompositeActor.h"
 #include "CompositeAssetUserData.h"
-#include "CompositeCoreSubsystem.h"
 #include "CompositeCoreSettings.h"
 #include "CompositeCVarOverrideManager.h"
 #include "CompositeRenderTargetPool.h"
+#include "CompositeWorldSubsystem.h"
 #include "Engine/Texture2D.h"
 #include "MediaTexture.h"
 #include "EngineUtils.h"
@@ -784,7 +784,7 @@ void UCompositeLayerPlate::TryRegisterCompositePrimitives(const TArray<UPrimitiv
 
 	if (PlateMode == ECompositePlateMode::CompositeMesh)
 	{
-		UCompositeCoreSubsystem* Subsystem = UWorld::GetSubsystem<UCompositeCoreSubsystem>(SourceWorld ? SourceWorld : GetWorld());
+		UCompositeWorldSubsystem* Subsystem = UWorld::GetSubsystem<UCompositeWorldSubsystem>(SourceWorld ? SourceWorld : GetWorld());
 		if (Subsystem)
 		{
 			// Primitives to be marked as holdout and rendered in the built-in composite custom render pass.
@@ -806,7 +806,7 @@ void UCompositeLayerPlate::UnregisterCompositePrimitives(const TArray<UPrimitive
 	}
 
 	// Always unregister no matter the current PlateMode
-	UCompositeCoreSubsystem* Subsystem = UWorld::GetSubsystem<UCompositeCoreSubsystem>(SourceWorld ? SourceWorld : GetWorld());
+	UCompositeWorldSubsystem* Subsystem = UWorld::GetSubsystem<UCompositeWorldSubsystem>(SourceWorld ? SourceWorld : GetWorld());
 	if (IsValid(Subsystem))
 	{
 		Subsystem->UnregisterPrimitives(InPrimitives, ECompositeCoreHoldoutManagement::None);
@@ -1016,34 +1016,34 @@ TSet<UPrimitiveComponent*> UCompositeLayerPlate::GetPrimitivesUsedByOtherActors(
 
 	TSet<UPrimitiveComponent*> Output;
 
-	const UWorld* World = GetWorld();
-	const ACompositeActor* ThisActor = GetTypedOuter<ACompositeActor>();
-
-	if (World && ThisActor)
+	if (const ACompositeActor* ThisActor = GetTypedOuter<ACompositeActor>())
 	{
-		for (TActorIterator<ACompositeActor> It(World); It; ++It)
+		if (UCompositeWorldSubsystem* Subsystem = UWorld::GetSubsystem<UCompositeWorldSubsystem>(GetWorld()))
 		{
-			const ACompositeActor* OtherActor = Cast<ACompositeActor>(*It);
-
-			if (!OtherActor || OtherActor == ThisActor || !OtherActor->IsRendering())
+			for (const TWeakObjectPtr<ACompositeActor> WeakActor : Subsystem->GetActors())
 			{
-				continue;
-			}
-			
-			for (const TObjectPtr<UCompositeLayerBase>& OtherLayer : OtherActor->GetCompositeLayers())
-			{
-				const UCompositeLayerPlate* OtherPlateLayer = Cast<UCompositeLayerPlate>(OtherLayer.Get());
+				const ACompositeActor* OtherActor = WeakActor.Get();
 
-				if (!OtherPlateLayer || !ensure(OtherPlateLayer != this) || !OtherPlateLayer->IsRendering())
+				if (!OtherActor || OtherActor == ThisActor || !OtherActor->IsRendering())
 				{
 					continue;
 				}
 
-				for (UPrimitiveComponent* OtherPrimitive : OtherPlateLayer->GetPrimitives())
+				for (const TObjectPtr<UCompositeLayerBase>& OtherLayer : OtherActor->GetCompositeLayers())
 				{
-					if (OtherPrimitive)
+					const UCompositeLayerPlate* OtherPlateLayer = Cast<UCompositeLayerPlate>(OtherLayer.Get());
+
+					if (!OtherPlateLayer || !ensure(OtherPlateLayer != this) || !OtherPlateLayer->IsRendering())
 					{
-						Output.Add(OtherPrimitive);
+						continue;
+					}
+
+					for (UPrimitiveComponent* OtherPrimitive : OtherPlateLayer->GetPrimitives())
+					{
+						if (OtherPrimitive)
+						{
+							Output.Add(OtherPrimitive);
+						}
 					}
 				}
 			}
diff --git a/Engine/Plugins/Compositing/Composite/Source/Composite/Public/CompositeActor.h b/Engine/Plugins/Compositing/Composite/Source/Composite/Public/CompositeActor.h
--- a/Engine/Plugins/Compositing/Composite/Source/Composite/Public/CompositeActor.h
+++ b/Engine/Plugins/Compositing/Composite/Source/Composite/Public/CompositeActor.h
@@ -14,7 +14,6 @@
 #define UE_API COMPOSITE_API
 
 class ACameraActor;
-class UCompositeCoreSubsystem;
 class UCompositeShadowReflectionCatcherComponent;
 class UPrimitiveComponent;
 class UCompositeSceneCapture2DComponent;
@@ -66,6 +65,7 @@ public:
 	//~ End UObject interface
 
 	//~ Begin AActor Interface
+	UE_API virtual void PreRegisterAllComponents() override;
 	UE_API virtual void PostRegisterAllComponents() override;
 	UE_API virtual void UnregisterAllComponents(bool bForReregister = false) override;
 	UE_API virtual void PostUnregisterAllComponents() override;
