Composite: Continuation of CL 50985365, refactor render resolution into layer-specific setting, which varies per layer type.

* Fixed issue with media/scene passes which are now correctly running at the media texture resolution. We dynamically recreate render targets once the media texture resolution is known/changes.
* Fixed various CIS flagged issues with zero or negative render target resolutions.

#Synced-CL 50934488

diff --git a/Engine/Plugins/Compositing/Composite/Source/Composite/Private/CompositeActor.cpp b/Engine/Plugins/Compositing/Composite/Source/Composite/Private/CompositeActor.cpp
--- a/Engine/Plugins/Compositing/Composite/Source/Composite/Private/CompositeActor.cpp
+++ b/Engine/Plugins/Compositing/Composite/Source/Composite/Private/CompositeActor.cpp
@@ -252,7 +252,6 @@ ACompositeActor::ACompositeActor(const FObjectInitializer& ObjectInitializer)
 	, ViewProjectionComponent{}
 	, bIsActive{true}
 	, bIsEnabled{true}
-	, RenderResolution{ FCompositeRenderTargetPool::DefaultSize }
 	, bEnableScreenSpaceReflections{true}
 	, bOverridesViewUserFlags{true}
 	, ViewUserFlags{1}
@@ -609,28 +608,6 @@ void ACompositeActor::PostEditUndo()
 	Super::PostEditUndo();
 }
 
-void ACompositeActor::PostTransacted(const FTransactionObjectEvent& TransactionEvent)
-{
-	Super::PostTransacted(TransactionEvent);
-
-	if (TransactionEvent.GetEventType() == ETransactionObjectEventType::UndoRedo)
-	{
-		const FName RenderResolutionName = GET_MEMBER_NAME_CHECKED(ThisClass, RenderResolution);
-		if (TransactionEvent.GetChangedProperties().Contains(RenderResolutionName))
-		{
-			/**
-			* When a resolution change is undone, the previous render targets are deserialized & replaced.
-			* This leaves previous assignments invalid so we make sure to return them to the pool.
-			*/
-			TInlineComponentArray<USceneCaptureComponent2D*> SceneCaptureComponents(this);
-			for (USceneCaptureComponent2D* CaptureComponent : SceneCaptureComponents)
-			{
-				FCompositeRenderTargetPool::Get().ReleaseAssigneeTargets(CaptureComponent);
-			}
-		}
-	}
-}
-
 void ACompositeActor::SetIsTemporarilyHiddenInEditor(bool bIsHidden)
 {
 	Super::SetIsTemporarilyHiddenInEditor(bIsHidden);
diff --git a/Engine/Plugins/Compositing/Composite/Source/Composite/Private/CompositeRenderTargetPool.cpp b/Engine/Plugins/Compositing/Composite/Source/Composite/Private/CompositeRenderTargetPool.cpp
--- a/Engine/Plugins/Compositing/Composite/Source/Composite/Private/CompositeRenderTargetPool.cpp
+++ b/Engine/Plugins/Compositing/Composite/Source/Composite/Private/CompositeRenderTargetPool.cpp
@@ -14,7 +14,7 @@ static TAutoConsoleVariable<int32> CVarCompositePoolFramesUntilFlush(
 	TEXT("This may leave render resources pooled, but never reclaimed. Auto-Flushing returns those resources. ")
 	TEXT("For values greater than zero, the pooling system will wait that number of frames before a target is considered 'stale'."));
 
-const FIntPoint FCompositeRenderTargetPool::DefaultSize = FIntPoint(960, 540);
+const FIntPoint FCompositeRenderTargetPool::DefaultSize = FIntPoint(1920, 1080);
 
 namespace UE
 {
@@ -92,6 +92,12 @@ TObjectPtr<UTextureRenderTarget2D> FCompositeRenderTargetPool::AcquireTarget(con
 
 bool FCompositeRenderTargetPool::ConditionalAcquireTarget(const UObject* InAssignee, TObjectPtr<UTextureRenderTarget2D>& InOutTarget, const FIntPoint& InSize, ETextureRenderTargetFormat InFormat)
 {
+	if (InSize.X <= 0 || InSize.Y <= 0)
+	{
+		UE_LOGFMT(LogComposite, Verbose, "Requested render target size of {0}x{1} is invalid.", InSize.X, InSize.Y);
+		return false;
+	}
+
 	if (IsValid(InOutTarget)
 		&& InOutTarget->SizeX == InSize.X
 		&& InOutTarget->SizeY == InSize.Y
diff --git a/Engine/Plugins/Compositing/Composite/Source/Composite/Private/Layers/CompositeLayerBase.cpp b/Engine/Plugins/Compositing/Composite/Source/Composite/Private/Layers/CompositeLayerBase.cpp
--- a/Engine/Plugins/Compositing/Composite/Source/Composite/Private/Layers/CompositeLayerBase.cpp
+++ b/Engine/Plugins/Compositing/Composite/Source/Composite/Private/Layers/CompositeLayerBase.cpp
@@ -3,7 +3,6 @@
 #include "Layers/CompositeLayerBase.h"
 
 #include "CompositeActor.h"
-#include "CompositeRenderTargetPool.h"
 #include "Components/PrimitiveComponent.h"
 #include "Components/SceneCaptureComponent2D.h"
 
@@ -156,18 +155,6 @@ void UCompositeLayerBase::AddChildPasses(UE::CompositeCore::FPassInputDecl& InBa
 	}
 }
 
-const FIntPoint UCompositeLayerBase::GetRenderResolution() const
-{
-	const ACompositeActor* CompositeActor = GetTypedOuter<ACompositeActor>();
-	
-	if (IsValid(CompositeActor))
-	{
-		return CompositeActor->RenderResolution;
-	}
-
-	return FCompositeRenderTargetPool::DefaultSize;
-}
-
 ECompositeCoreMergeOp UCompositeLayerBase::GetMergeOperation(const FCompositeTraversalContext& InContext) const
 {
 	return InContext.bIsSolo ? ECompositeCoreMergeOp::None : Operation;
diff --git a/Engine/Plugins/Compositing/Composite/Source/Composite/Private/Layers/CompositeLayerPlate.cpp b/Engine/Plugins/Compositing/Composite/Source/Composite/Private/Layers/CompositeLayerPlate.cpp
--- a/Engine/Plugins/Compositing/Composite/Source/Composite/Private/Layers/CompositeLayerPlate.cpp
+++ b/Engine/Plugins/Compositing/Composite/Source/Composite/Private/Layers/CompositeLayerPlate.cpp
@@ -84,6 +84,31 @@ namespace UE
 				constexpr bool bRestoreVisibility = false;
 				HandleRemovedCompositeMeshes(InPreEditMeshes, {}, bRestoreVisibility);
 			}
+
+			/** Return the texture size, depending on the texture type. */
+			static FIntPoint GetTextureSize(const UTexture* InTexture)
+			{
+				if (const UMediaTexture* MediaTexture = Cast<UMediaTexture>(InTexture))
+				{
+					const int32 Width = MediaTexture->GetWidth();
+					const int32 Height = MediaTexture->GetHeight();
+
+					if (Width > 0 && Height > 0)
+					{
+						return FIntPoint(Width, Height);
+					}
+				}
+				else if (const UTexture2D* Texture2D = Cast<UTexture2D>(InTexture))
+				{
+					return FIntPoint(Texture2D->GetSizeX(), Texture2D->GetSizeY());
+				}
+				else if (const UTextureRenderTarget2D* RenderTarget2D = Cast<UTextureRenderTarget2D>(InTexture))
+				{
+					return FIntPoint(RenderTarget2D->SizeX, RenderTarget2D->SizeY);
+				}
+
+				return FCompositeRenderTargetPool::DefaultSize;
+			}
 		}
 	}
 }
@@ -158,8 +183,9 @@ void UCompositeLayerPlate::Tick(float DeltaTime)
 		bPassesHaveChanged = true;
 	}
 
-	if (bPassesHaveChanged)
+	if (bPassesHaveChanged || bRebindRenderTarget)
 	{
+		bRebindRenderTarget = false;
 		UpdateCompositeMeshes();
 	}
 }
@@ -400,6 +426,12 @@ void UCompositeLayerPlate::PostTransacted(const FTransactionObjectEvent& InTrans
 				// and then attempt to open the corresponding media source if Texture is a media profile media texture
 				TryCloseMediaProfileSource();
 				TryOpenMediaProfileSource();
+
+				/**
+				* When a resolution change is undone, the previous render targets are deserialized & replaced.
+				* This leaves previous assignments invalid so we make sure to return them to the pool.
+				*/
+				FCompositeRenderTargetPool::Get().ReleaseAssigneeTargets(this);
 			}
 			
 			UpdateCompositeMeshes();
@@ -1055,9 +1087,20 @@ TSet<UPrimitiveComponent*> UCompositeLayerPlate::GetPrimitivesUsedByOtherActors(
 
 TObjectPtr<UTextureRenderTarget2D> UCompositeLayerPlate::GetOrCreateRenderTarget(TObjectPtr<UTextureRenderTarget2D>& InRenderTarget) const
 {
+	const FIntPoint TextureSize = UE::Composite::Private::GetTextureSize(Texture);
+
+	if (IsValid(InRenderTarget) && FIntPoint(InRenderTarget->SizeX, InRenderTarget->SizeY) != TextureSize)
+	{
+		// Automatically release our render targets, such that they are resized to match the dynamic media texture resolution.
+		FCompositeRenderTargetPool::Get().ReleaseTarget(InRenderTarget);
+		
+		// For sequencer use, it is preferable to warmup the media texture since this rebind can occur on the next tick.
+		bRebindRenderTarget = true;
+	}
+
 	if (!IsValid(InRenderTarget))
 	{
-		FCompositeRenderTargetPool::Get().ConditionalAcquireTarget(this, InRenderTarget, GetRenderResolution());
+		FCompositeRenderTargetPool::Get().ConditionalAcquireTarget(this, InRenderTarget, TextureSize);
 	}
 
 	return InRenderTarget;
diff --git a/Engine/Plugins/Compositing/Composite/Source/Composite/Private/Layers/CompositeLayerSceneCapture.cpp b/Engine/Plugins/Compositing/Composite/Source/Composite/Private/Layers/CompositeLayerSceneCapture.cpp
--- a/Engine/Plugins/Compositing/Composite/Source/Composite/Private/Layers/CompositeLayerSceneCapture.cpp
+++ b/Engine/Plugins/Compositing/Composite/Source/Composite/Private/Layers/CompositeLayerSceneCapture.cpp
@@ -7,12 +7,14 @@
 #include "Components/CompositeSceneCapture2DComponent.h"
 #include "CompositeRenderTargetPool.h"
 #include "Engine/TextureRenderTarget2D.h"
+#include "Misc/TransactionObjectEvent.h"
 #include "UObject/Package.h"
 
 UCompositeLayerSceneCapture::UCompositeLayerSceneCapture(const FObjectInitializer& ObjectInitializer)
 	: Super(ObjectInitializer)
-	, bVisibleInSceneCaptureOnly{true} // By default, registered meshes will no longer be visible in the main render
-	, bCustomRenderPass{false}
+	, RenderTargetResolution{ FCompositeRenderTargetPool::DefaultSize }
+	, bVisibleInSceneCaptureOnly{ true } // By default, registered meshes will no longer be visible in the main render
+	, bCustomRenderPass{ false }
 {
 	Operation = ECompositeCoreMergeOp::Over;
 
@@ -81,6 +83,10 @@ void UCompositeLayerSceneCapture::PostEditChangeProperty(FPropertyChangedEvent&
 	{
 		SetVisibleInSceneCaptureOnly(bVisibleInSceneCaptureOnly);
 	}
+	else if (PropertyName == GET_MEMBER_NAME_CHECKED(ThisClass, RenderTargetResolution))
+	{
+		SetRenderTargetResolution(RenderTargetResolution);
+	}
 }
 
 void UCompositeLayerSceneCapture::PreEditUndo()
@@ -98,6 +104,24 @@ void UCompositeLayerSceneCapture::PostEditUndo()
 
 	SetCustomRenderPass(bCustomRenderPass);
 }
+
+void UCompositeLayerSceneCapture::PostTransacted(const FTransactionObjectEvent& TransactionEvent)
+{
+	Super::PostTransacted(TransactionEvent);
+
+	if (TransactionEvent.GetEventType() == ETransactionObjectEventType::UndoRedo &&
+		TransactionEvent.GetChangedProperties().Contains(GET_MEMBER_NAME_CHECKED(ThisClass, RenderTargetResolution)))
+	{
+		/**
+		* When a resolution change is undone, the previous render targets are deserialized & replaced.
+		* This leaves previous assignments invalid so we make sure to return them to the pool.
+		*/
+		if (UCompositeSceneCapture2DComponent* CaptureComponent = GetSceneCaptureComponent())
+		{
+			FCompositeRenderTargetPool::Get().ReleaseAssigneeTargets(CaptureComponent);
+		}
+	}
+}
 #endif
 
 void UCompositeLayerSceneCapture::OnEndOfFrameUpdate(UWorld* InWorld)
@@ -130,7 +154,7 @@ bool UCompositeLayerSceneCapture::GetProxy(FCompositeTraversalContext& InContext
 		return false;
 	}
 
-	FCompositeRenderTargetPool::Get().ConditionalAcquireTarget(CaptureComponent, CaptureComponent->TextureTarget, GetRenderResolution());
+	FCompositeRenderTargetPool::Get().ConditionalAcquireTarget(CaptureComponent, CaptureComponent->TextureTarget, RenderTargetResolution);
 
 	FResourceMetadata Metadata = {};
 	Metadata.bInvertedAlpha = true;
@@ -213,6 +237,12 @@ void UCompositeLayerSceneCapture::SetVisibleInSceneCaptureOnly(bool bInVisible)
 	UpdatePrimitiveVisibilityState(bVisibleInSceneCaptureOnly);
 }
 
+void UCompositeLayerSceneCapture::SetRenderTargetResolution(FIntPoint InRenderTargetResolution)
+{
+	RenderTargetResolution.X = FMath::Max(InRenderTargetResolution.X, 1);
+	RenderTargetResolution.Y = FMath::Max(InRenderTargetResolution.Y, 1);
+}
+
 UCompositeSceneCapture2DComponent* UCompositeLayerSceneCapture::GetSceneCaptureComponent() const
 {
 	ACompositeActor* CompositeActor = GetTypedOuter<ACompositeActor>();
diff --git a/Engine/Plugins/Compositing/Composite/Source/Composite/Private/Layers/CompositeLayerShadowReflection.cpp b/Engine/Plugins/Compositing/Composite/Source/Composite/Private/Layers/CompositeLayerShadowReflection.cpp
--- a/Engine/Plugins/Compositing/Composite/Source/Composite/Private/Layers/CompositeLayerShadowReflection.cpp
+++ b/Engine/Plugins/Compositing/Composite/Source/Composite/Private/Layers/CompositeLayerShadowReflection.cpp
@@ -9,9 +9,11 @@
 #include "CompositeModule.h"
 #include "CompositeRenderTargetPool.h"
 #include "Containers/StaticArray.h"
+#include "Misc/TransactionObjectEvent.h"
 
 UCompositeLayerShadowReflection::UCompositeLayerShadowReflection(const FObjectInitializer& ObjectInitializer)
 	: Super(ObjectInitializer)
+	, RenderTargetResolution{ 960, 540 } // Lower resolution for better performance by default
 {
 	Operation = ECompositeCoreMergeOp::Multiply;
 
@@ -83,6 +85,10 @@ void UCompositeLayerShadowReflection::PostEditChangeProperty(FPropertyChangedEve
 	{
 		SetActors(MoveTemp(Actors));
 	}
+	else if (PropertyName == GET_MEMBER_NAME_CHECKED(ThisClass, RenderTargetResolution))
+	{
+		SetRenderTargetResolution(RenderTargetResolution);
+	}
 }
 
 void UCompositeLayerShadowReflection::PreEditUndo()
@@ -99,6 +105,27 @@ void UCompositeLayerShadowReflection::PostEditUndo()
 
 	SetActors(MoveTemp(Actors));
 }
+
+void UCompositeLayerShadowReflection::PostTransacted(const FTransactionObjectEvent& TransactionEvent)
+{
+	Super::PostTransacted(TransactionEvent);
+
+	if (TransactionEvent.GetEventType() == ETransactionObjectEventType::UndoRedo &&
+		TransactionEvent.GetChangedProperties().Contains(GET_MEMBER_NAME_CHECKED(ThisClass, RenderTargetResolution)))
+	{
+		/**
+		* When a resolution change is undone, the previous render targets are deserialized & replaced.
+		* This leaves previous assignments invalid so we make sure to return them to the pool.
+		*/
+		for (const TWeakObjectPtr<UCompositeShadowReflectionCatcherComponent>& CaptureWeak : CachedSceneCaptures)
+		{
+			if (UCompositeShadowReflectionCatcherComponent* Capture = CaptureWeak.Get())
+			{
+				FCompositeRenderTargetPool::Get().ReleaseAssigneeTargets(Capture);
+			}
+		}
+	}
+}
 #endif
 
 void UCompositeLayerShadowReflection::OnEndOfFrameUpdate(UWorld* InWorld)
@@ -147,7 +174,7 @@ bool UCompositeLayerShadowReflection::GetProxy(FCompositeTraversalContext& InCon
 		for (int32 Index = 0; Index < 2; ++Index)
 		{
 			UCompositeShadowReflectionCatcherComponent* Component = CachedSceneCaptures[Index].Get();
-			FCompositeRenderTargetPool::Get().ConditionalAcquireTarget(Component, Component->TextureTarget, GetRenderResolution());
+			FCompositeRenderTargetPool::Get().ConditionalAcquireTarget(Component, Component->TextureTarget, RenderTargetResolution);
 
 			FResourceMetadata Metadata = {};
 			Metadata.bInvertedAlpha = true;
@@ -210,6 +237,12 @@ void UCompositeLayerShadowReflection::SetActors(TArray<TSoftObjectPtr<AActor>> I
 	UpdateSceneCaptureComponents(*CachedSceneCaptures[1], Actors);
 }
 
+void UCompositeLayerShadowReflection::SetRenderTargetResolution(FIntPoint InRenderTargetResolution)
+{
+	RenderTargetResolution.X = FMath::Max(InRenderTargetResolution.X, 1);
+	RenderTargetResolution.Y = FMath::Max(InRenderTargetResolution.Y, 1);
+}
+
 void UCompositeLayerShadowReflection::SetEnabled(bool bInEnabled)
 {
 	Super::SetEnabled(bInEnabled);
diff --git a/Engine/Plugins/Compositing/Composite/Source/Composite/Private/Layers/CompositeLayerSingleLightShadow.cpp b/Engine/Plugins/Compositing/Composite/Source/Composite/Private/Layers/CompositeLayerSingleLightShadow.cpp
--- a/Engine/Plugins/Compositing/Composite/Source/Composite/Private/Layers/CompositeLayerSingleLightShadow.cpp
+++ b/Engine/Plugins/Compositing/Composite/Source/Composite/Private/Layers/CompositeLayerSingleLightShadow.cpp
@@ -5,6 +5,7 @@
 #include "Camera/CameraComponent.h"
 #include "CompositeActor.h"
 #include "CompositeRenderTargetPool.h"
+#include "Misc/TransactionObjectEvent.h"
 #include "Passes/CompositePassMaterial.h"
 #include "Components/CompositeSceneCapture2DComponent.h"
 #include "Engine/DirectionalLight.h"
@@ -131,6 +132,7 @@ UCompositeLayerSingleLightShadow::UCompositeLayerSingleLightShadow(const FObject
 	, OrthographicWidth{ DEFAULT_ORTHOWIDTH }
 	, ShadowMapResolution{ 2048 }
 	, ShadowStrength{ 0.5f }
+	, RenderTargetResolution{ FCompositeRenderTargetPool::DefaultSize }
 {
 	Operation = ECompositeCoreMergeOp::Multiply;
 
@@ -168,6 +170,10 @@ void UCompositeLayerSingleLightShadow::PostEditChangeProperty(FPropertyChangedEv
 	{
 		SetShadowCastingActors(MoveTemp(ShadowCastingActors));
 	}
+	else if (PropertyName == GET_MEMBER_NAME_CHECKED(ThisClass, RenderTargetResolution))
+	{
+		SetRenderTargetResolution(RenderTargetResolution);
+	}
 }
 
 void UCompositeLayerSingleLightShadow::PostEditUndo()
@@ -176,6 +182,30 @@ void UCompositeLayerSingleLightShadow::PostEditUndo()
 
 	SetShadowCastingActors(MoveTemp(ShadowCastingActors));
 }
+
+void UCompositeLayerSingleLightShadow::SetRenderTargetResolution(FIntPoint InRenderTargetResolution)
+{
+	RenderTargetResolution.X = FMath::Max(InRenderTargetResolution.X, 1);
+	RenderTargetResolution.Y = FMath::Max(InRenderTargetResolution.Y, 1);
+}
+
+void UCompositeLayerSingleLightShadow::PostTransacted(const FTransactionObjectEvent& TransactionEvent)
+{
+	Super::PostTransacted(TransactionEvent);
+
+	if (TransactionEvent.GetEventType() == ETransactionObjectEventType::UndoRedo &&
+		TransactionEvent.GetChangedProperties().Contains(GET_MEMBER_NAME_CHECKED(ThisClass, RenderTargetResolution)))
+	{
+		/**
+		* When a resolution change is undone, the previous render targets are deserialized & replaced.
+		* This leaves previous assignments invalid so we make sure to return them to the pool.
+		*/
+		if (UCompositeSceneCapture2DComponent* Capture = CachedSceneDepthCapture.Get())
+		{
+			FCompositeRenderTargetPool::Get().ReleaseAssigneeTargets(Capture);
+		}
+	}
+}
 #endif
 
 void UCompositeLayerSingleLightShadow::OnEndOfFrameUpdate(UWorld* InWorld)
@@ -276,14 +306,14 @@ bool UCompositeLayerSingleLightShadow::GetProxy(FCompositeTraversalContext& InCo
 
 		// Acquire render target
 		FCompositeRenderTargetPool::Get().ConditionalAcquireTarget(CachedShadowMapCapture.Get(), CachedShadowMapCapture->TextureTarget, FIntPoint(ShadowMapResolution, ShadowMapResolution), RTF_R32f);
-		FCompositeRenderTargetPool::Get().ConditionalAcquireTarget(CachedSceneDepthCapture.Get(), CachedSceneDepthCapture->TextureTarget, GetRenderResolution(), RTF_R32f);
+		FCompositeRenderTargetPool::Get().ConditionalAcquireTarget(CachedSceneDepthCapture.Get(), CachedSceneDepthCapture->TextureTarget, RenderTargetResolution, RTF_R32f);
 	}
 	
 	// Setup child shadow pass proxy
 	FCompositeCorePassProxy* InputProxy;
 
 	{
-		// Display-resolution scene depth without jitter as an external texture.
+		// Scene depth without jitter as an external texture.
 		const ResourceId SceneDepthId = InContext.FindOrCreateExternalTexture(CachedSceneDepthCapture->TextureTarget, { .DebugName = TEXT("CompositSceneDepthTex") });
 
 		// Shadow map texture
diff --git a/Engine/Plugins/Compositing/Composite/Source/Composite/Public/CompositeActor.h b/Engine/Plugins/Compositing/Composite/Source/Composite/Public/CompositeActor.h
--- a/Engine/Plugins/Compositing/Composite/Source/Composite/Public/CompositeActor.h
+++ b/Engine/Plugins/Compositing/Composite/Source/Composite/Public/CompositeActor.h
@@ -79,7 +79,6 @@ public:
 	UE_API virtual void PostEditChangeChainProperty(struct FPropertyChangedChainEvent& PropertyChangedChainEvent) override;
 	UE_API virtual void PreEditUndo() override;
 	UE_API virtual void PostEditUndo() override;
-	UE_API virtual void PostTransacted(const FTransactionObjectEvent& TransactionEvent) override;
 	UE_API virtual void SetIsTemporarilyHiddenInEditor(bool bIsHidden) override;
 	UE_API virtual bool ShouldTickIfViewportsOnly() const override;
 #endif // WITH_EDITOR
@@ -161,10 +160,6 @@ private:
 	bool bIsEnabled;
 
 public:
-	/** Composite render target resolution, used by managed scene captures. */
-	UPROPERTY(EditAnywhere, BlueprintReadWrite, Interp, Category = "Composite")
-	FIntPoint RenderResolution;
-
 private:
 	/** The primary camera used for the composite. Used by the camera view projection component for continuous view-projection matrix updates. */
 	UPROPERTY(EditInstanceOnly, BlueprintGetter = GetCamera, BlueprintSetter = SetCamera, Category = "Composite", meta = (AllowPrivateAccess, UseComponentPicker, AllowAnyActor, AllowedClasses = "/Script/Engine.CameraComponent, /Script/CinematicCamera.CineCameraComponent"))
diff --git a/Engine/Plugins/Compositing/Composite/Source/Composite/Public/Layers/CompositeLayerBase.h b/Engine/Plugins/Compositing/Composite/Source/Composite/Public/Layers/CompositeLayerBase.h
--- a/Engine/Plugins/Compositing/Composite/Source/Composite/Public/Layers/CompositeLayerBase.h
+++ b/Engine/Plugins/Compositing/Composite/Source/Composite/Public/Layers/CompositeLayerBase.h
@@ -102,9 +102,6 @@ protected:
 	/** Fixed number of inputs used by layer merge passes. */
 	static constexpr int32 FixedNumLayerInputs = 2;
 
-	/** Convenience function that returns the composite actor's render resolution. */
-	const FIntPoint GetRenderResolution() const;
-
 	/** Convenience function that returns the merge operation depending on the traversal context. */
 	ECompositeCoreMergeOp GetMergeOperation(const FCompositeTraversalContext& InContext) const;
 
diff --git a/Engine/Plugins/Compositing/Composite/Source/Composite/Public/Layers/CompositeLayerPlate.h b/Engine/Plugins/Compositing/Composite/Source/Composite/Public/Layers/CompositeLayerPlate.h
--- a/Engine/Plugins/Compositing/Composite/Source/Composite/Public/Layers/CompositeLayerPlate.h
+++ b/Engine/Plugins/Compositing/Composite/Source/Composite/Public/Layers/CompositeLayerPlate.h
@@ -263,6 +263,9 @@ private:
 	/** Cached number of valid scene passes. */
 	int32 CachedValidScenePasses = 0;
 
+	/** Toggled when a render target is resized to trigger a mesh material texture parameter rebind in Tick. */
+	mutable bool bRebindRenderTarget = false;
+
 	friend class FCompositeLayerPlateCustomization;
 };
 
diff --git a/Engine/Plugins/Compositing/Composite/Source/Composite/Public/Layers/CompositeLayerSceneCapture.h b/Engine/Plugins/Compositing/Composite/Source/Composite/Public/Layers/CompositeLayerSceneCapture.h
--- a/Engine/Plugins/Compositing/Composite/Source/Composite/Public/Layers/CompositeLayerSceneCapture.h
+++ b/Engine/Plugins/Compositing/Composite/Source/Composite/Public/Layers/CompositeLayerSceneCapture.h
@@ -37,6 +37,7 @@ public:
 	UE_API virtual void PostEditChangeProperty(FPropertyChangedEvent& PropertyChangedEvent) override;
 	UE_API virtual void PreEditUndo() override;
 	UE_API virtual void PostEditUndo() override;
+	UE_API virtual void PostTransacted(const FTransactionObjectEvent& TransactionEvent) override;
 #endif
 	//~ End UObject Interface
 
@@ -70,6 +71,14 @@ public:
 	UFUNCTION(BlueprintSetter)
 	UE_API void SetVisibleInSceneCaptureOnly(bool bInVisible);
 
+	/** Get the render target resolution for the scene capture. */
+	UFUNCTION(BlueprintGetter)
+	UE_API FIntPoint GetRenderTargetResolution() const { return RenderTargetResolution; }
+
+	/** Set the render target resolution for the scene capture. */
+	UFUNCTION(BlueprintSetter)
+	UE_API void SetRenderTargetResolution(FIntPoint InRenderTargetResolution);
+
 	/** Sub-passes applied during post-processing. */
 	UPROPERTY(EditAnywhere, BlueprintReadWrite, Instanced, Category = "", meta = (DisallowedClasses = "/Script/Composite.CompositePassDistortion", DisplayAfter = "bCustomRenderPass"))
 	TArray<TObjectPtr<UCompositePassBase>> LayerPasses;
@@ -99,6 +108,10 @@ private:
 	TMap<TWeakObjectPtr<UPrimitiveComponent>, bool> CachedVisibilityInSceneCapture;
 #endif
 
+	/** Render target resolution for the scene capture. */
+	UPROPERTY(EditAnywhere, BlueprintGetter = GetRenderTargetResolution, BlueprintSetter = SetRenderTargetResolution, Interp, Category = "", meta = (AllowPrivateAccess = true))
+	FIntPoint RenderTargetResolution;
+
 	/** Visibility setting applied to registered actors. By default, registered meshes will only be visible in scene captures and no longer in the main render. */
 	UPROPERTY(EditAnywhere, BlueprintGetter = IsVisibleInSceneCaptureOnly, BlueprintSetter = SetVisibleInSceneCaptureOnly, Category = "", meta = (AllowPrivateAccess = true))
 	bool bVisibleInSceneCaptureOnly;
diff --git a/Engine/Plugins/Compositing/Composite/Source/Composite/Public/Layers/CompositeLayerShadowReflection.h b/Engine/Plugins/Compositing/Composite/Source/Composite/Public/Layers/CompositeLayerShadowReflection.h
--- a/Engine/Plugins/Compositing/Composite/Source/Composite/Public/Layers/CompositeLayerShadowReflection.h
+++ b/Engine/Plugins/Compositing/Composite/Source/Composite/Public/Layers/CompositeLayerShadowReflection.h
@@ -71,6 +71,7 @@ public:
 	UE_API virtual void PostEditChangeProperty(FPropertyChangedEvent& PropertyChangedEvent) override;
 	UE_API virtual void PreEditUndo() override;
 	UE_API virtual void PostEditUndo() override;
+	UE_API virtual void PostTransacted(const FTransactionObjectEvent& TransactionEvent) override;
 #endif
 	//~ End UObject Interface
 
@@ -80,6 +81,14 @@ public:
 	/** Getter function to override, returning pass layer proxies to be passed to the render thread. (Proxy objects should be allocated from the provided allocator.) */
 	UE_API virtual bool GetProxy(FCompositeTraversalContext& InContext, FSceneRenderingBulkObjectAllocator& InFrameAllocator, FCompositeCorePassProxy*& OutProxy) const override;
 
+	/** Get the render target resolution for the scene captures. */
+	UFUNCTION(BlueprintGetter)
+	UE_API FIntPoint GetRenderTargetResolution() const { return RenderTargetResolution; }
+
+	/** Set the render target resolution for the scene captures. */
+	UFUNCTION(BlueprintSetter)
+	UE_API void SetRenderTargetResolution(FIntPoint InRenderTargetResolution);
+
 	/** Get reflection & shadow caster actors. */
 	UFUNCTION(BlueprintGetter)
 	UE_API const TArray<TSoftObjectPtr<AActor>> GetActors() const;
@@ -110,6 +119,10 @@ private:
 	TStaticArray<TWeakObjectPtr<UCompositeShadowReflectionCatcherComponent>, 2> FindOrCreateSceneCapturePair(ACompositeActor& InOuter) const;
 
 private:
+	/** Render target resolution for the scene captures. */
+	UPROPERTY(EditAnywhere, BlueprintGetter = GetRenderTargetResolution, BlueprintSetter = SetRenderTargetResolution, Interp, Category = "", meta = (AllowPrivateAccess = true))
+	FIntPoint RenderTargetResolution;
+
 	/** List of reflection & shadow caster actors. */
 	UPROPERTY(EditAnywhere, BlueprintGetter = GetActors, BlueprintSetter = SetActors, Category = "")
 	TArray<TSoftObjectPtr<AActor>> Actors;
diff --git a/Engine/Plugins/Compositing/Composite/Source/Composite/Public/Layers/CompositeLayerSingleLightShadow.h b/Engine/Plugins/Compositing/Composite/Source/Composite/Public/Layers/CompositeLayerSingleLightShadow.h
--- a/Engine/Plugins/Compositing/Composite/Source/Composite/Public/Layers/CompositeLayerSingleLightShadow.h
+++ b/Engine/Plugins/Compositing/Composite/Source/Composite/Public/Layers/CompositeLayerSingleLightShadow.h
@@ -47,6 +47,7 @@ public:
 #if WITH_EDITOR
 	UE_API virtual void PostEditChangeProperty(FPropertyChangedEvent& PropertyChangedEvent) override;
 	UE_API virtual void PostEditUndo() override;
+	UE_API virtual void PostTransacted(const FTransactionObjectEvent& TransactionEvent) override;
 #endif
 	//~ End UObject Interface
 
@@ -87,6 +88,14 @@ public:
 	UPROPERTY(EditAnywhere, BlueprintReadWrite, Category = "", meta = (EditCondition = "!bAutomaticBounds", EditConditionHides))
 	float OrthographicWidth;
 
+	/** Get the render target resolution for the scene depth capture. */
+	UFUNCTION(BlueprintGetter)
+	UE_API FIntPoint GetRenderTargetResolution() const { return RenderTargetResolution; }
+
+	/** Set the render target resolution for the scene depth capture. */
+	UFUNCTION(BlueprintSetter)
+	UE_API void SetRenderTargetResolution(FIntPoint InRenderTargetResolution);
+
 	/**
 	* Resolution of the shadow map texture. Lower resolutions will cause more blurry shadows.
 	*/
@@ -118,6 +127,10 @@ private:
 	void GetShadowMatrices(UCompositeSceneCapture2DComponent& ShadowMapCapture, FMatrix44f& OutShadowMatrix, FVector4f& OutShadowInvDeviceZToWorldZ) const;
 
 private:
+	/** Render target resolution for the scene depth capture. */
+	UPROPERTY(EditAnywhere, BlueprintGetter = GetRenderTargetResolution, BlueprintSetter = SetRenderTargetResolution, Interp, Category = "", meta = (AllowPrivateAccess = true))
+	FIntPoint RenderTargetResolution;
+
 	/** List of shadow casting actors (required). */
 	UPROPERTY(EditAnywhere, BlueprintGetter = GetShadowCastingActors, BlueprintSetter = SetShadowCastingActors, Category = "", meta = (AllowPrivateAccess = true))
 	TArray<TSoftObjectPtr<AActor>> ShadowCastingActors;
diff --git a/Engine/Plugins/Compositing/Composite/Source/CompositeEditor/Private/Customizations/CompositeLayerSceneCaptureCustomization.cpp b/Engine/Plugins/Compositing/Composite/Source/CompositeEditor/Private/Customizations/CompositeLayerSceneCaptureCustomization.cpp
--- a/Engine/Plugins/Compositing/Composite/Source/CompositeEditor/Private/Customizations/CompositeLayerSceneCaptureCustomization.cpp
+++ b/Engine/Plugins/Compositing/Composite/Source/CompositeEditor/Private/Customizations/CompositeLayerSceneCaptureCustomization.cpp
@@ -2,6 +2,7 @@
 
 #include "CompositeLayerSceneCaptureCustomization.h"
 
+#include "CompositeRenderTargetResolutionCustomization.h"
 #include "DetailCategoryBuilder.h"
 #include "DetailLayoutBuilder.h"
 #include "DetailWidgetRow.h"
@@ -26,6 +27,12 @@ void FCompositeLayerSceneCaptureCustomization::CustomizeDetails(const TSharedPtr
 
 void FCompositeLayerSceneCaptureCustomization::CustomizeDetails(IDetailLayoutBuilder& DetailLayout)
 {
+	DetailLayout.RegisterInstancedCustomPropertyTypeLayout(NAME_IntPoint, FOnGetPropertyTypeCustomizationInstance::CreateLambda([]
+		{
+			return MakeShared<FResolutionTypeCustomization>();
+		}),
+		MakeShared<FResolutionPropertyIdentifier>());
+
 	DetailLayout.HideProperty(GET_MEMBER_NAME_CHECKED(UCompositeLayerSceneCapture, Actors));
 	
 	TArray<TWeakObjectPtr<UCompositeLayerSceneCapture>> Objects = DetailLayout.GetObjectsOfTypeBeingCustomized<UCompositeLayerSceneCapture>();
diff --git a/Engine/Plugins/Compositing/Composite/Source/CompositeEditor/Private/Customizations/CompositeLayerShadowReflectionCustomization.cpp b/Engine/Plugins/Compositing/Composite/Source/CompositeEditor/Private/Customizations/CompositeLayerShadowReflectionCustomization.cpp
--- a/Engine/Plugins/Compositing/Composite/Source/CompositeEditor/Private/Customizations/CompositeLayerShadowReflectionCustomization.cpp
+++ b/Engine/Plugins/Compositing/Composite/Source/CompositeEditor/Private/Customizations/CompositeLayerShadowReflectionCustomization.cpp
@@ -2,6 +2,7 @@
 
 #include "CompositeLayerShadowReflectionCustomization.h"
 
+#include "CompositeRenderTargetResolutionCustomization.h"
 #include "ActorTreeItem.h"
 #include "CompositeMeshActor.h"
 #include "DetailCategoryBuilder.h"
@@ -29,6 +30,12 @@ void FCompositeLayerShadowReflectionCustomization::CustomizeDetails(const TShare
 
 void FCompositeLayerShadowReflectionCustomization::CustomizeDetails(IDetailLayoutBuilder& DetailLayout)
 {
+	DetailLayout.RegisterInstancedCustomPropertyTypeLayout(NAME_IntPoint, FOnGetPropertyTypeCustomizationInstance::CreateLambda([]
+		{
+			return MakeShared<FResolutionTypeCustomization>();
+		}),
+		MakeShared<FResolutionPropertyIdentifier>());
+
 	DetailLayout.HideProperty(GET_MEMBER_NAME_CHECKED(UCompositeLayerShadowReflection, Actors));
 	
 	TArray<TWeakObjectPtr<UCompositeLayerShadowReflection>> Objects = DetailLayout.GetObjectsOfTypeBeingCustomized<UCompositeLayerShadowReflection>();
diff --git a/Engine/Plugins/Compositing/Composite/Source/CompositeEditor/Private/Customizations/CompositeLayerSingleLightShadowCustomization.cpp b/Engine/Plugins/Compositing/Composite/Source/CompositeEditor/Private/Customizations/CompositeLayerSingleLightShadowCustomization.cpp
--- a/Engine/Plugins/Compositing/Composite/Source/CompositeEditor/Private/Customizations/CompositeLayerSingleLightShadowCustomization.cpp
+++ b/Engine/Plugins/Compositing/Composite/Source/CompositeEditor/Private/Customizations/CompositeLayerSingleLightShadowCustomization.cpp
@@ -2,6 +2,7 @@
 
 #include "CompositeLayerSingleLightShadowCustomization.h"
 
+#include "CompositeRenderTargetResolutionCustomization.h"
 #include "DetailCategoryBuilder.h"
 #include "DetailLayoutBuilder.h"
 #include "DetailWidgetRow.h"
@@ -26,6 +27,12 @@ void FCompositeLayerSingleLightShadowCustomization::CustomizeDetails(const TShar
 
 void FCompositeLayerSingleLightShadowCustomization::CustomizeDetails(IDetailLayoutBuilder& DetailLayout)
 {
+	DetailLayout.RegisterInstancedCustomPropertyTypeLayout(NAME_IntPoint, FOnGetPropertyTypeCustomizationInstance::CreateLambda([]
+		{
+			return MakeShared<FResolutionTypeCustomization>();
+		}),
+		MakeShared<FResolutionPropertyIdentifier>());
+
 	DetailLayout.HideProperty(GET_MEMBER_NAME_CHECKED(UCompositeLayerSingleLightShadow, ShadowCastingActors));
 	
 	TArray<TWeakObjectPtr<UCompositeLayerSingleLightShadow>> Objects = DetailLayout.GetObjectsOfTypeBeingCustomized<UCompositeLayerSingleLightShadow>();
diff --git a/Engine/Plugins/Compositing/Composite/Source/CompositeEditor/Private/Customizations/CompositeRenderTargetResolutionCustomization.cpp b/Engine/Plugins/Compositing/Composite/Source/CompositeEditor/Private/Customizations/CompositeRenderTargetResolutionCustomization.cpp
new file mode 100644
--- /dev/null
+++ b/Engine/Plugins/Compositing/Composite/Source/CompositeEditor/Private/Customizations/CompositeRenderTargetResolutionCustomization.cpp
@@ -0,0 +1,256 @@
+// Copyright Epic Games, Inc. All Rights Reserved.
+
+#include "CompositeRenderTargetResolutionCustomization.h"
+
+#include "DetailLayoutBuilder.h"
+#include "Math/UnrealMathUtility.h"
+#include "Styling/AppStyle.h"
+#include "Widgets/Images/SImage.h"
+#include "Widgets/Input/SButton.h"
+
+#define LOCTEXT_NAMESPACE "CompositeResolutionCustomization"
+
+TArray<TSharedPtr<FResolutionTypeCustomization::FNamedResolution>> FResolutionTypeCustomization::NamedResolutions =
+{
+	MakeShared<FNamedResolution>(LOCTEXT("540pResolutionName", "540p (qHD) - 960x540"), FIntPoint(960, 540), LOCTEXT("Resolution540pToolTip", "Quarter High Definition")),
+	MakeShared<FNamedResolution>(LOCTEXT("720pResolutionName", "720p (HD) - 1280x720"), FIntPoint(1280, 720), LOCTEXT("Resolution720pToolTip", "High Definition")),
+	MakeShared<FNamedResolution>(LOCTEXT("1080pResolutionName", "1080p (FHD) - 1920x1080"), FIntPoint(1920, 1080), LOCTEXT("Resolution1080pToolTip", "Full high definition")),
+	MakeShared<FNamedResolution>(LOCTEXT("1440pResolutionName", "1440p (QHD) - 2560x1440"), FIntPoint(2560, 1440), LOCTEXT("Resolution1440pToolTip", "Quad high definition")),
+	MakeShared<FNamedResolution>(LOCTEXT("2160pResolutionName", "2160p (4K UHD) - 3840x2160"), FIntPoint(3840, 2160), LOCTEXT("Resolution2160pToolTip", "Ultra high definition (4k)")),
+	MakeShared<FNamedResolution>(LOCTEXT("CustomResolutionName", "Custom"), FIntPoint::ZeroValue, FText::GetEmpty())
+};
+
+void FResolutionTypeCustomization::CustomizeHeader(TSharedRef<IPropertyHandle> InPropertyHandle, FDetailWidgetRow& InHeaderRow, IPropertyTypeCustomizationUtils& InCustomizationUtils)
+{
+	PropertyHandle = InPropertyHandle;
+
+	XPropertyHandle = InPropertyHandle->GetChildHandle(GET_MEMBER_NAME_CHECKED(FIntPoint, X));
+	YPropertyHandle = InPropertyHandle->GetChildHandle(GET_MEMBER_NAME_CHECKED(FIntPoint, Y));
+
+	FIntPoint InitialResolution = FIntPoint::ZeroValue;
+	FPropertyAccess::Result XResult = XPropertyHandle->GetValue(InitialResolution.X);
+	FPropertyAccess::Result YResult = YPropertyHandle->GetValue(InitialResolution.Y);
+
+	int32 InitialResolutionIndex = NamedResolutions.IndexOfByPredicate([&InitialResolution](const FNamedResolutionPtr& NamedResolution)
+	{
+		return NamedResolution->Resolution == InitialResolution;
+	});
+
+	if (InitialResolutionIndex == INDEX_NONE)
+	{
+		InitialResolutionIndex = NamedResolutions.Num() - 1;
+	}
+
+	InHeaderRow
+	.NameContent()
+	[
+		PropertyHandle->CreatePropertyNameWidget()
+	]
+	.ValueContent()
+	[
+		SAssignNew(ResolutionComboBox, SComboBox<FNamedResolutionPtr>)
+		.OptionsSource(&NamedResolutions)
+		.OnGenerateWidget_Lambda([](FNamedResolutionPtr InItem)
+		{
+			return SNew(STextBlock)
+				.Font(IDetailLayoutBuilder::GetDetailFont())
+				.Text(InItem->Name)
+				.ToolTipText(InItem->ToolTip);
+		})
+		.OnSelectionChanged(this, &FResolutionTypeCustomization::SetResolution)
+		.InitiallySelectedItem(NamedResolutions[InitialResolutionIndex])
+		[
+			SNew(STextBlock)
+			.Font(IDetailLayoutBuilder::GetDetailFont())
+			.Text(this, &FResolutionTypeCustomization::GetResolutionText)
+		]
+	];
+
+	InHeaderRow.OverrideResetToDefault(FResetToDefaultOverride::Create(
+		FIsResetToDefaultVisible::CreateLambda([this](TSharedPtr<IPropertyHandle>)
+		{
+			return ResolutionComboBox->GetSelectedItem() != NamedResolutions[0];
+		}), FResetToDefaultHandler::CreateLambda([this](TSharedPtr<IPropertyHandle>)
+		{
+			ResolutionComboBox->SetSelectedItem(NamedResolutions[0]);
+		})));
+}
+
+void FResolutionTypeCustomization::CustomizeChildren(TSharedRef<IPropertyHandle> InPropertyHandle, IDetailChildrenBuilder& InChildBuilder, IPropertyTypeCustomizationUtils& InCustomizationUtils)
+{
+	InChildBuilder.AddCustomRow(InPropertyHandle->GetPropertyDisplayName())
+	.NameContent()
+	[
+		XPropertyHandle->CreatePropertyNameWidget()
+	]
+	.ValueContent()
+	[
+		XPropertyHandle->CreatePropertyValueWidget()
+	]
+	.ExtensionContent()
+	[
+		SNew(SButton)
+		.OnClicked_Lambda([this]()
+		{
+			if (LockedAspectRatio.IsSet())
+			{
+				LockedAspectRatio.Reset();
+			}
+			else
+			{
+				double AspectRatio;
+				if (GetAspectRatio(AspectRatio))
+				{
+					LockedAspectRatio = AspectRatio;
+				}
+			}
+
+			return FReply::Handled();
+		})
+		.ContentPadding(FMargin(0, 0, 4, 0))
+		.ButtonStyle(FAppStyle::Get(), "NoBorder")
+		[
+			SNew(SImage)
+			.ColorAndOpacity(FSlateColor::UseForeground())
+			.Image_Lambda([this]
+			{
+				return LockedAspectRatio.IsSet() ? FAppStyle::GetBrush(TEXT("Icons.Link")) : FAppStyle::GetBrush(TEXT("Icons.Unlink"));
+			})
+			.ToolTipText_Lambda([this]
+			{
+				return LockedAspectRatio.IsSet() ?
+				FText::Format(LOCTEXT("LockedAspectRatioTooltipFormat", "Click to unlock Aspect Ratio ({0})"), FText::AsNumber(LockedAspectRatio.GetValue())) :
+				LOCTEXT("UnlockedAspectRatioTooltip", "Click to lock the Aspect Ratio");
+			})
+		]
+	];
+
+	InChildBuilder.AddCustomRow(InPropertyHandle->GetPropertyDisplayName())
+	.NameContent()
+	[
+		YPropertyHandle->CreatePropertyNameWidget()
+	]
+	.ValueContent()
+	[
+		YPropertyHandle->CreatePropertyValueWidget()
+	];
+
+	XPropertyHandle->SetOnPropertyValueChanged(FSimpleDelegate::CreateLambda([this]
+	{
+		if (!bSettingNamedResolution)
+		{
+			ResolutionComboBox->SetSelectedItem(NamedResolutions.Last());
+		}
+
+		if (LockedAspectRatio.IsSet())
+		{
+			const double AspectRatio = LockedAspectRatio.GetValue();
+			
+			if (AspectRatio > UE_SMALL_NUMBER)
+			{
+				SetDimensionFromAspectRatio(XPropertyHandle, YPropertyHandle, 1.0 / AspectRatio);
+			}
+		}
+	}));
+
+	YPropertyHandle->SetOnPropertyValueChanged(FSimpleDelegate::CreateLambda([this]
+	{
+		if (!bSettingNamedResolution)
+		{
+			ResolutionComboBox->SetSelectedItem(NamedResolutions.Last());
+		}
+
+		if (LockedAspectRatio.IsSet())
+		{
+			SetDimensionFromAspectRatio(YPropertyHandle, XPropertyHandle, LockedAspectRatio.GetValue());
+		}
+	}));
+}
+
+FText FResolutionTypeCustomization::GetResolutionText() const
+{
+	FIntPoint Resolution;
+	FPropertyAccess::Result XResult = PropertyHandle->GetChildHandle(GET_MEMBER_NAME_CHECKED(FIntPoint, X))->GetValue(Resolution.X);
+	FPropertyAccess::Result YResult = PropertyHandle->GetChildHandle(GET_MEMBER_NAME_CHECKED(FIntPoint, Y))->GetValue(Resolution.Y);
+
+	if (XResult == FPropertyAccess::Fail || YResult == FPropertyAccess::Fail)
+	{
+		return FText::GetEmpty();
+	}
+
+	if (XResult == FPropertyAccess::MultipleValues || YResult == FPropertyAccess::MultipleValues)
+	{
+		return LOCTEXT("MultipleValues", "Multiple Values");
+	}
+
+	FNamedResolutionPtr SelectedResolution = ResolutionComboBox->GetSelectedItem();
+	if (SelectedResolution->Resolution != FIntPoint::ZeroValue)
+	{
+		return SelectedResolution->Name;
+	}
+
+	FNumberFormattingOptions FormattingOptions;
+	FormattingOptions.UseGrouping = false;
+	return FText::Format(LOCTEXT("CustomResolutionTextFormat", "Custom - {0}x{1}"),
+		FText::AsNumber(Resolution.X, &FormattingOptions),
+		FText::AsNumber(Resolution.Y, &FormattingOptions));
+}
+
+void FResolutionTypeCustomization::SetResolution(FNamedResolutionPtr InNamedResolution, ESelectInfo::Type SelectInfo)
+{
+	if (InNamedResolution->Resolution != FIntPoint::ZeroValue)
+	{
+		TGuardValue<bool> SettingNamedResolution(bSettingNamedResolution, true);
+		PropertyHandle->GetChildHandle(GET_MEMBER_NAME_CHECKED(FIntPoint, X))->SetValue(InNamedResolution->Resolution.X);
+		PropertyHandle->GetChildHandle(GET_MEMBER_NAME_CHECKED(FIntPoint, Y))->SetValue(InNamedResolution->Resolution.Y);
+	}
+}
+
+bool FResolutionTypeCustomization::GetAspectRatio(double& OutAspectRatio) const
+{
+	int32 X = 0;
+	int32 Y = 0;
+
+	if (XPropertyHandle->GetValue(X) != FPropertyAccess::Success)
+	{
+		return false;
+	}
+
+	if (YPropertyHandle->GetValue(Y) != FPropertyAccess::Success)
+	{
+		return false;
+	}
+
+	if (Y == 0)
+	{
+		return false;
+	}
+
+	OutAspectRatio = (double)X / (double)Y;
+	return true;
+}
+
+void FResolutionTypeCustomization::SetDimensionFromAspectRatio(const TSharedPtr<IPropertyHandle>& SrcDimHandle, const TSharedPtr<IPropertyHandle>& DestDimHandle, double AspectRatio)
+{
+	if (bAspectRatioRecursionGuard)
+	{
+		return;
+	}
+
+	TGuardValue<bool> RecursionGuard(bAspectRatioRecursionGuard, true);
+
+	int32 SrcDim;
+	if (SrcDimHandle->GetValue(SrcDim) != FPropertyAccess::Success)
+	{
+		return;
+	}
+
+	const double ScaledSrcDim = SrcDim * AspectRatio;
+
+	// Round to the nearest even number
+	const int32 DestDim = FMath::RoundToInt(ScaledSrcDim / 2.0) * 2;
+
+	DestDimHandle->SetValue(DestDim);
+}
+
+#undef LOCTEXT_NAMESPACE
diff --git a/Engine/Plugins/Compositing/Composite/Source/CompositeEditor/Private/Customizations/CompositeRenderTargetResolutionCustomization.h b/Engine/Plugins/Compositing/Composite/Source/CompositeEditor/Private/Customizations/CompositeRenderTargetResolutionCustomization.h
new file mode 100644
--- /dev/null
+++ b/Engine/Plugins/Compositing/Composite/Source/CompositeEditor/Private/Customizations/CompositeRenderTargetResolutionCustomization.h
@@ -0,0 +1,67 @@
+// Copyright Epic Games, Inc. All Rights Reserved.
+
+#pragma once
+
+#include "DetailWidgetRow.h"
+#include "IDetailChildrenBuilder.h"
+#include "IPropertyTypeCustomization.h"
+#include "Widgets/Input/SComboBox.h"
+
+/**
+ * Property type customization for FIntPoint render target resolution properties.
+ * Displays a resolution combobox with named presets and an aspect ratio lock.
+ */
+class FResolutionTypeCustomization : public IPropertyTypeCustomization
+{
+public:
+	struct FNamedResolution
+	{
+		FText Name = FText::GetEmpty();
+		FIntPoint Resolution = FIntPoint::ZeroValue;
+		FText ToolTip = FText::GetEmpty();
+
+		FNamedResolution(const FText& InName, const FIntPoint& InResolution, const FText& InToolTip)
+			: Name(InName)
+			, Resolution(InResolution)
+			, ToolTip(InToolTip)
+		{ }
+	};
+	using FNamedResolutionPtr = TSharedPtr<FNamedResolution>;
+
+	static TArray<FNamedResolutionPtr> NamedResolutions;
+
+	virtual void CustomizeHeader(TSharedRef<IPropertyHandle> InPropertyHandle, FDetailWidgetRow& InHeaderRow, IPropertyTypeCustomizationUtils& InCustomizationUtils) override;
+	virtual void CustomizeChildren(TSharedRef<IPropertyHandle> InPropertyHandle, IDetailChildrenBuilder& InChildBuilder, IPropertyTypeCustomizationUtils& InCustomizationUtils) override;
+
+private:
+	FText GetResolutionText() const;
+	void SetResolution(FNamedResolutionPtr InNamedResolution, ESelectInfo::Type SelectInfo);
+	bool GetAspectRatio(double& OutAspectRatio) const;
+	void SetDimensionFromAspectRatio(const TSharedPtr<IPropertyHandle>& SrcDimHandle, const TSharedPtr<IPropertyHandle>& DestDimHandle, double AspectRatio);
+
+private:
+	TSharedPtr<IPropertyHandle> PropertyHandle;
+	TSharedPtr<IPropertyHandle> XPropertyHandle;
+	TSharedPtr<IPropertyHandle> YPropertyHandle;
+
+	TSharedPtr<SComboBox<FNamedResolutionPtr>> ResolutionComboBox;
+
+	TOptional<double> LockedAspectRatio = TOptional<double>();
+
+	bool bAspectRatioRecursionGuard = false;
+	bool bSettingNamedResolution = false;
+};
+
+/**
+ * Property type identifier that matches FIntPoint properties named "RenderTargetResolution",
+ * selecting them for FResolutionTypeCustomization.
+ */
+class FResolutionPropertyIdentifier : public IPropertyTypeIdentifier
+{
+public:
+	virtual bool IsPropertyTypeCustomized(const IPropertyHandle& PropertyHandle) const override
+	{
+		return PropertyHandle.GetProperty()->GetNameCPP() == TEXT("RenderTargetResolution");
+	}
+};
+
diff --git a/Engine/Plugins/Compositing/Composite/Source/CompositeEditor/Private/UI/SCompositeEditorPanel.cpp b/Engine/Plugins/Compositing/Composite/Source/CompositeEditor/Private/UI/SCompositeEditorPanel.cpp
--- a/Engine/Plugins/Compositing/Composite/Source/CompositeEditor/Private/UI/SCompositeEditorPanel.cpp
+++ b/Engine/Plugins/Compositing/Composite/Source/CompositeEditor/Private/UI/SCompositeEditorPanel.cpp
@@ -39,302 +39,11 @@ namespace CompositingEditorPanel
 	}
 }
 
-class FResolutionTypeCustomization : public IPropertyTypeCustomization
-{
-private:
-	struct FNamedResolution
-	{
-		FText Name = FText::GetEmpty();
-		FIntPoint Resolution = FIntPoint::ZeroValue;
-		FText ToolTip = FText::GetEmpty();
-			
-		FNamedResolution(const FText& InName, const FIntPoint& InResolution, const FText& InToolTip)
-			: Name(InName)
-			, Resolution(InResolution)
-			, ToolTip(InToolTip)
-		{ }
-	};
-	using FNamedResolutionPtr = TSharedPtr<FNamedResolution>;
-	
-	static TArray<FNamedResolutionPtr> NamedResolutions;
-	
-public:
-	virtual void CustomizeHeader(TSharedRef<IPropertyHandle> InPropertyHandle, FDetailWidgetRow& InHeaderRow, IPropertyTypeCustomizationUtils& InCustomizationUtils) override
-	{
-		PropertyHandle = InPropertyHandle;
-
-		XPropertyHandle = InPropertyHandle->GetChildHandle(GET_MEMBER_NAME_CHECKED(FIntPoint, X));
-		YPropertyHandle = InPropertyHandle->GetChildHandle(GET_MEMBER_NAME_CHECKED(FIntPoint, Y));
-
-		FIntPoint InitialResolution = FIntPoint::ZeroValue;
-		FPropertyAccess::Result XResult = XPropertyHandle->GetValue(InitialResolution.X);
-		FPropertyAccess::Result YResult = YPropertyHandle->GetValue(InitialResolution.Y);
-
-		int32 InitialResolutionIndex = NamedResolutions.IndexOfByPredicate([&InitialResolution](const FNamedResolutionPtr& NamedResolution)
-		{
-			return NamedResolution->Resolution == InitialResolution;
-		});
-
-		if (InitialResolutionIndex == INDEX_NONE)
-		{
-			InitialResolutionIndex = NamedResolutions.Num() - 1;
-		}
-		
-		InHeaderRow
-		.NameContent()
-		[
-			PropertyHandle->CreatePropertyNameWidget()
-		]
-		.ValueContent()
-		[
-			SAssignNew(ResolutionComboBox, SComboBox<FNamedResolutionPtr>)
-			.OptionsSource(&NamedResolutions)
-			.OnGenerateWidget_Lambda([](FNamedResolutionPtr InItem)
-			{
-				return SNew(STextBlock)
-					.Font(IDetailLayoutBuilder::GetDetailFont())
-					.Text(InItem->Name)
-					.ToolTipText(InItem->ToolTip);
-			})
-			.OnSelectionChanged(this, &FResolutionTypeCustomization::SetResolution)
-			.InitiallySelectedItem(NamedResolutions[InitialResolutionIndex])
-			[
-				SNew(STextBlock)
-				.Font(IDetailLayoutBuilder::GetDetailFont())
-				.Text(this, &FResolutionTypeCustomization::GetResolutionText)
-			]
-		];
-
-		InHeaderRow.OverrideResetToDefault(FResetToDefaultOverride::Create(
-			FIsResetToDefaultVisible::CreateLambda([this](TSharedPtr<IPropertyHandle>)
-			{
-				return ResolutionComboBox->GetSelectedItem() != NamedResolutions[0];
-			}), FResetToDefaultHandler::CreateLambda([this](TSharedPtr<IPropertyHandle>)
-			{
-				ResolutionComboBox->SetSelectedItem(NamedResolutions[0]);
-			})));
-	}
-
-	virtual void CustomizeChildren(TSharedRef<IPropertyHandle> InPropertyHandle, IDetailChildrenBuilder& InChildBuilder, IPropertyTypeCustomizationUtils& InCustomizationUtils) override
-	{
-		InChildBuilder.AddCustomRow(InPropertyHandle->GetPropertyDisplayName())
-		.NameContent()
-		[
-			XPropertyHandle->CreatePropertyNameWidget()
-		]
-		.ValueContent()
-		[
-			XPropertyHandle->CreatePropertyValueWidget()
-		]
-		.ExtensionContent()
-		[
-			SNew(SButton)
-			.OnClicked_Lambda([this]()
-			{
-				if (LockedAspectRatio.IsSet())
-				{
-					LockedAspectRatio.Reset();
-				}
-				else
-				{
-					double AspectRatio;
-					if (GetAspectRatio(AspectRatio))
-					{
-						LockedAspectRatio = AspectRatio;
-					}
-				}
-				
-				return FReply::Handled();
-			})
-			.ContentPadding(FMargin(0, 0, 4, 0))
-			.ButtonStyle(FAppStyle::Get(), "NoBorder")
-			[
-				SNew(SImage)
-				.ColorAndOpacity(FSlateColor::UseForeground())
-				.Image_Lambda([this]
-				{
-					return LockedAspectRatio.IsSet() ? FAppStyle::GetBrush(TEXT("Icons.Link")) : FAppStyle::GetBrush(TEXT("Icons.Unlink"));
-				})
-				.ToolTipText_Lambda([this]
-				{
-					return LockedAspectRatio.IsSet() ?
-					FText::Format(LOCTEXT("LockedAspectRatioTooltipFormat", "Click to unlock Aspect Ratio ({0})"), FText::AsNumber(LockedAspectRatio.GetValue())) :
-					LOCTEXT("UnlockedAspectRatioTooltip", "Click to lock the Aspect Ratio");
-				})
-			]
-		];
-		
-		InChildBuilder.AddCustomRow(InPropertyHandle->GetPropertyDisplayName())
-		.NameContent()
-		[
-			YPropertyHandle->CreatePropertyNameWidget()
-		]
-		.ValueContent()
-		[
-			YPropertyHandle->CreatePropertyValueWidget()
-		];
-
-		XPropertyHandle->SetOnPropertyValueChanged(FSimpleDelegate::CreateLambda([this]
-		{
-			if (!bSettingNamedResolution)
-			{
-				ResolutionComboBox->SetSelectedItem(NamedResolutions.Last());
-			}
-
-			if (LockedAspectRatio.IsSet())
-			{
-				SetDimensionFromAspectRatio(XPropertyHandle, YPropertyHandle, 1.0 / LockedAspectRatio.GetValue());
-			}
-		}));
-		
-		YPropertyHandle->SetOnPropertyValueChanged(FSimpleDelegate::CreateLambda([this]
-		{
-			if (!bSettingNamedResolution)
-			{
-				ResolutionComboBox->SetSelectedItem(NamedResolutions.Last());
-			}
-
-			if (LockedAspectRatio.IsSet())
-			{
-				SetDimensionFromAspectRatio(YPropertyHandle, XPropertyHandle, LockedAspectRatio.GetValue());
-			}
-		}));
-	}
-	
-private:
-	FText GetResolutionText() const
-	{
-		FIntPoint Resolution;
-		FPropertyAccess::Result XResult = PropertyHandle->GetChildHandle(GET_MEMBER_NAME_CHECKED(FIntPoint, X))->GetValue(Resolution.X);
-		FPropertyAccess::Result YResult = PropertyHandle->GetChildHandle(GET_MEMBER_NAME_CHECKED(FIntPoint, Y))->GetValue(Resolution.Y);
-
-		if (XResult == FPropertyAccess::Fail || YResult == FPropertyAccess::Fail)
-		{
-			return FText::GetEmpty();
-		}
-
-		if (XResult == FPropertyAccess::MultipleValues || YResult == FPropertyAccess::MultipleValues)
-		{
-			return LOCTEXT("MultipleValues", "Multiple Values");
-		}
-
-		FNamedResolutionPtr SelectedResolution = ResolutionComboBox->GetSelectedItem();
-		if (SelectedResolution->Resolution != FIntPoint::ZeroValue)
-		{
-			return SelectedResolution->Name;
-		}
-
-		FNumberFormattingOptions FormattingOptions;
-		FormattingOptions.UseGrouping = false;
-		return FText::Format(LOCTEXT("CustomResolutionTextFormat", "Custom - {0}x{1}"),
-			FText::AsNumber(Resolution.X, &FormattingOptions),
-			FText::AsNumber(Resolution.Y, &FormattingOptions));
-	}
-
-	void SetResolution(FNamedResolutionPtr InNamedResolution, ESelectInfo::Type SelectInfo)
-	{
-		if (InNamedResolution->Resolution != FIntPoint::ZeroValue)
-		{
-			TGuardValue<bool> SettingNamedResolution(bSettingNamedResolution, true);
-			PropertyHandle->GetChildHandle(GET_MEMBER_NAME_CHECKED(FIntPoint, X))->SetValue(InNamedResolution->Resolution.X);
-			PropertyHandle->GetChildHandle(GET_MEMBER_NAME_CHECKED(FIntPoint, Y))->SetValue(InNamedResolution->Resolution.Y);
-		}
-	}
-
-	bool GetAspectRatio(double& OutAspectRatio) const
-	{
-		int32 X = 0;
-		int32 Y = 0;
-
-		if (XPropertyHandle->GetValue(X) != FPropertyAccess::Success)
-		{
-			return false;
-		}
-
-		if (YPropertyHandle->GetValue(Y) != FPropertyAccess::Success)
-		{
-			return false;
-		}
-
-		if (Y == 0)
-		{
-			return false;
-		}
-
-		OutAspectRatio = (double)X / (double)Y;
-		return true;
-	}
-
-	void SetDimensionFromAspectRatio(const TSharedPtr<IPropertyHandle>& SrcDimHandle, const TSharedPtr<IPropertyHandle>& DestDimHandle, double AspectRatio)
-	{
-		if (bAspectRatioRecursionGuard)
-		{
-			return;
-		}
-
-		TGuardValue<bool> RecursionGuard(bAspectRatioRecursionGuard, true);
-		
-		int32 SrcDim;
-		if (SrcDimHandle->GetValue(SrcDim) != FPropertyAccess::Success)
-		{
-			return;
-		}
-
-		const double ScaledSrcDim = SrcDim * AspectRatio;
-
-		// Round to the nearest even number
-		int32 DestDim = FMath::FloorToInt(ScaledSrcDim);
-		if (DestDim % 2 != 0)
-		{
-			DestDim = FMath::CeilToInt(ScaledSrcDim);
-		}
-
-		DestDimHandle->SetValue(DestDim);
-	}
-	
-private:
-	TSharedPtr<IPropertyHandle> PropertyHandle;
-	TSharedPtr<IPropertyHandle> XPropertyHandle;
-	TSharedPtr<IPropertyHandle> YPropertyHandle;
-
-	TSharedPtr<SComboBox<FNamedResolutionPtr>> ResolutionComboBox;
-	
-	TOptional<double> LockedAspectRatio = TOptional<double>();
-
-	bool bAspectRatioRecursionGuard = false;
-	bool bSettingNamedResolution = false;
-};
-
-TArray<TSharedPtr<FResolutionTypeCustomization::FNamedResolution>> FResolutionTypeCustomization::NamedResolutions =
-{
-	MakeShared<FNamedResolution>(LOCTEXT("540pResolutionName", "540p (qHD) - 960x540"), FIntPoint(960, 540), LOCTEXT("Resolution540pToolTip", "Quarter High Definition")),
-	MakeShared<FNamedResolution>(LOCTEXT("720pResolutionName", "720p (HD) - 1280x720"), FIntPoint(1280, 720), LOCTEXT("Resolution720pToolTip", "High Definition")),
-	MakeShared<FNamedResolution>(LOCTEXT("1080pResolutionName", "1080p (FHD) - 1920x1080"), FIntPoint(1920, 1080), LOCTEXT("Resolution1080pToolTip", "Full high definition")),
-	MakeShared<FNamedResolution>(LOCTEXT("1440pResolutionName", "1440p (QHD) - 2560x1440"), FIntPoint(2560, 1440), LOCTEXT("Resolution1440pToolTip", "Quad high definition")),
-	MakeShared<FNamedResolution>(LOCTEXT("2160pResolutionName", "2160p (4K UHD) - 3840x2160"), FIntPoint(3840, 2160), LOCTEXT("Resolution2160pToolTip", "Ultra high definition (4k)")),
-	MakeShared<FNamedResolution>(LOCTEXT("CustomResolutionName", "Custom"), FIntPoint::ZeroValue, FText::GetEmpty())
-};
-
 class FCompositeActorPanelDetailCustomization : public IDetailCustomization
 {
-private:
-	class FResolutionPropertyIdentifier : public IPropertyTypeIdentifier
-	{
-		virtual bool IsPropertyTypeCustomized(const IPropertyHandle& PropertyHandle) const override
-		{
-			return PropertyHandle.GetProperty()->GetNameCPP() == GET_MEMBER_NAME_STRING_CHECKED(ACompositeActor, RenderResolution);
-		}
-	};
-	
 public:
 	virtual void CustomizeDetails(IDetailLayoutBuilder& DetailBuilder) override
 	{
-		DetailBuilder.RegisterInstancedCustomPropertyTypeLayout(NAME_IntPoint, FOnGetPropertyTypeCustomizationInstance::CreateLambda([]
-			{
-				return MakeShared<FResolutionTypeCustomization>();
-			}),
-			MakeShared<FResolutionPropertyIdentifier>());
-		
 		DetailBuilder.GetObjectsBeingCustomized(ObjectsBeingCustomized);
 		
 		// Hide all categories besides the Composite category
@@ -401,8 +110,6 @@ public:
 
 private:
 	TArray<TWeakObjectPtr<UObject>> ObjectsBeingCustomized;
-
-	TSharedPtr<IPropertyHandle> RenderResolutionHandle;
 };
 
 void SCompositeEditorPanel::RegisterTabSpawner()
