Composure: Follow up fix to CL 50619036: prevent hiding composite meshes on layer/actor deletion.

#Synced-CL 51061694

diff --git a/Engine/Plugins/Compositing/Composite/Source/Composite/Private/CompositeActor.cpp b/Engine/Plugins/Compositing/Composite/Source/Composite/Private/CompositeActor.cpp
--- a/Engine/Plugins/Compositing/Composite/Source/Composite/Private/CompositeActor.cpp
+++ b/Engine/Plugins/Compositing/Composite/Source/Composite/Private/CompositeActor.cpp
@@ -370,6 +370,9 @@ void ACompositeActor::UnregisterAllComponents(bool bForReregister)
 		return;
 	}
 	
+	// We don't want to hide meshes upon deletion
+	TGuardValue<bool> Guard(UCompositeLayerPlate::GControlMeshVisibility, false);
+
 	// State removal is done without any conditions to prevent stale state
 	PropagateStateChange(false);
 }
diff --git a/Engine/Plugins/Compositing/Composite/Source/Composite/Private/Layers/CompositeLayerPlate.cpp b/Engine/Plugins/Compositing/Composite/Source/Composite/Private/Layers/CompositeLayerPlate.cpp
--- a/Engine/Plugins/Compositing/Composite/Source/Composite/Private/Layers/CompositeLayerPlate.cpp
+++ b/Engine/Plugins/Compositing/Composite/Source/Composite/Private/Layers/CompositeLayerPlate.cpp
@@ -33,9 +33,10 @@
 
 const FLazyName UCompositeLayerPlate::CompositeTextureName = FLazyName(TEXT("CompositeTexture"));
 
-TAutoConsoleVariable<bool> CVarCompositePlateControlMeshVisibility(
+bool UCompositeLayerPlate::GControlMeshVisibility = true;
+FAutoConsoleVariableRef CVarCompositePlateControlMeshVisibility(
 	TEXT("Composite.Plate.ControlMeshVisibility"),
-	true,
+	UCompositeLayerPlate::GControlMeshVisibility,
 	TEXT("When enabled, the plate layers automatically control the visibility of registered composite meshes")
 );
 
@@ -231,6 +232,9 @@ void UCompositeLayerPlate::OnRemoved(ACompositeActor* LastOwner)
 
 	TryCloseMediaProfileSource();
 
+	// We don't want to hide meshes upon deletion
+	TGuardValue<bool> Guard(GControlMeshVisibility, false);
+
 	constexpr bool bApply = false;
 	PropagateStateChange(bApply, LastOwner ? LastOwner->GetWorld() : nullptr);
 }
@@ -263,10 +267,8 @@ const TArray<TSoftObjectPtr<AActor>> UCompositeLayerPlate::GetCompositeMeshes()
 
 void UCompositeLayerPlate::SetCompositeMeshes(TArray<TSoftObjectPtr<AActor>> InCompositeMeshes)
 {
-	const bool bControlMeshVisibility = CVarCompositePlateControlMeshVisibility.GetValueOnGameThread();
-
 	// First, we untrack current meshes
-	UE::Composite::Private::HandleRemovedCompositeMeshes(CompositeMeshes, InCompositeMeshes, bControlMeshVisibility);
+	UE::Composite::Private::HandleRemovedCompositeMeshes(CompositeMeshes, InCompositeMeshes, GControlMeshVisibility);
 
 	CompositeMeshes = MoveTemp(InCompositeMeshes);
 	
@@ -367,9 +369,7 @@ void UCompositeLayerPlate::PostEditChangeProperty(FPropertyChangedEvent& Propert
 
 		if (bChangeWasRemoval)
 		{
-			const bool bControlMeshVisibility = CVarCompositePlateControlMeshVisibility.GetValueOnGameThread();
-
-			UE::Composite::Private::HandleRemovedCompositeMeshes(PreEditCompositeMeshes, CompositeMeshes, bControlMeshVisibility);
+			UE::Composite::Private::HandleRemovedCompositeMeshes(PreEditCompositeMeshes, CompositeMeshes, GControlMeshVisibility);
 		}
 		else
 		{
@@ -930,8 +930,6 @@ void UCompositeLayerPlate::UpdatePrimitiveComponent(UPrimitiveComponent& InPrimi
 
 void UCompositeLayerPlate::PropagateStateChange(bool bApply, const UWorld* SourceWorld) const
 {
-	const bool bControlMeshVisibility = CVarCompositePlateControlMeshVisibility.GetValueOnGameThread();
-
 	if (bApply && IsRendering())
 	{
 		if (CompositeMeshes.IsEmpty())
@@ -944,7 +942,7 @@ void UCompositeLayerPlate::PropagateStateChange(bool bApply, const UWorld* Sourc
 
 		const TArray<UPrimitiveComponent*> Primitives = GetPrimitives();
 
-		if (bControlMeshVisibility)
+		if (GControlMeshVisibility)
 		{
 			for (UPrimitiveComponent* Primitive : Primitives)
 			{
@@ -995,7 +993,7 @@ void UCompositeLayerPlate::PropagateStateChange(bool bApply, const UWorld* Sourc
 
 			UnregisterCompositePrimitives(Primitives, SourceWorld);
 
-			if (bControlMeshVisibility)
+			if (GControlMeshVisibility)
 			{
 				TSet<UPrimitiveComponent*> UsedPrimitives = GetPrimitivesUsedByOtherActors();
 
diff --git a/Engine/Plugins/Compositing/Composite/Source/Composite/Public/Layers/CompositeLayerPlate.h b/Engine/Plugins/Compositing/Composite/Source/Composite/Public/Layers/CompositeLayerPlate.h
--- a/Engine/Plugins/Compositing/Composite/Source/Composite/Public/Layers/CompositeLayerPlate.h
+++ b/Engine/Plugins/Compositing/Composite/Source/Composite/Public/Layers/CompositeLayerPlate.h
@@ -151,6 +151,9 @@ public:
 	UPROPERTY(EditAnywhere, BlueprintReadWrite, Instanced, Category = "", meta = (DisallowedClasses = "/Script/Composite.CompositePassDistortion"))
 	TArray<TObjectPtr<UCompositePassBase>> LayerPasses;
 
+	/** Composite.Plate.ControlMeshVisibility console variable value. */
+	static bool GControlMeshVisibility;
+
 private:
 	/** The plate sampling & rendering mode: direct texture sampling or separate composite mesh render. */
 	UPROPERTY(EditAnywhere, BlueprintGetter = GetPlateMode, BlueprintSetter = SetPlateMode, Category = "", AdvancedDisplay, meta = (AllowPrivateAccess = true, DisplayName = "Mode"))
