Composure: Refactor layer class to derive from pass.

#Synced-CL 51065964

#Synced-CL 51065964

diff --git a/Engine/Plugins/Compositing/Composite/Source/Composite/Private/CompositeActor.cpp b/Engine/Plugins/Compositing/Composite/Source/Composite/Private/CompositeActor.cpp
--- a/Engine/Plugins/Compositing/Composite/Source/Composite/Private/CompositeActor.cpp
+++ b/Engine/Plugins/Compositing/Composite/Source/Composite/Private/CompositeActor.cpp
@@ -111,7 +111,8 @@ namespace UE
 					}
 
 					FCompositeCorePassProxy* OutProxy;
-					if (Layer->GetProxy(Context, InFrameAllocator, OutProxy))
+					UE::CompositeCore::FPassInputDecl LayerInputDecl;
+					if (Layer->GetProxy(LayerInputDecl, Context, InFrameAllocator, OutProxy))
 					{
 						OutProxies.Add(OutProxy);
 						Context.bIsFirstPass = false;
diff --git a/Engine/Plugins/Compositing/Composite/Source/Composite/Private/Layers/CompositeLayerBase.cpp b/Engine/Plugins/Compositing/Composite/Source/Composite/Private/Layers/CompositeLayerBase.cpp
--- a/Engine/Plugins/Compositing/Composite/Source/Composite/Private/Layers/CompositeLayerBase.cpp
+++ b/Engine/Plugins/Compositing/Composite/Source/Composite/Private/Layers/CompositeLayerBase.cpp
@@ -5,29 +5,30 @@
 #include "CompositeActor.h"
 #include "Components/PrimitiveComponent.h"
 #include "Components/SceneCaptureComponent2D.h"
+#include "UObject/FortniteMainBranchObjectVersion.h"
 
 UCompositeLayerBase::UCompositeLayerBase(const FObjectInitializer& ObjectInitializer)
 	: Super(ObjectInitializer)
 {
-	// Ensure instances are always transactional
-	if (!HasAnyFlags(RF_ClassDefaultObject))
-	{
-		SetFlags(RF_Transactional);
-	}
+	// Marked RF_Transactional if non-CDO in parent UCompositePassBase
+}
 
-#if WITH_EDITOR
-	Name = GetClass()->GetName();
-	
-	const int32 NameNumber = GetFName().GetNumber();
-	if (NameNumber > 1)
+UCompositeLayerBase::~UCompositeLayerBase() = default;
+
+void UCompositeLayerBase::PostLoad()
+{
+	Super::PostLoad();
+
+#if WITH_EDITORONLY_DATA
+	const bool bDeprecatedNameMigration = GetLinkerCustomVersion(FFortniteMainBranchObjectVersion::GUID) < FFortniteMainBranchObjectVersion::Composite_LayerNameDeprecation;
+
+	if (bDeprecatedNameMigration && !Name_DEPRECATED.IsEmpty() && (DisplayName.IsEmpty() || DisplayName == GetClass()->GetName()))
 	{
-		Name.AppendInt(NameNumber);
+		DisplayName = Name_DEPRECATED;
 	}
 #endif
 }
 
-UCompositeLayerBase::~UCompositeLayerBase() = default;
-
 void UCompositeLayerBase::BeginDestroy()
 {
 	Super::BeginDestroy();
@@ -75,14 +76,13 @@ void UCompositeLayerBase::PostEditChangeProperty(FPropertyChangedEvent& Property
 }
 #endif //WITH_EDITOR
 
-bool UCompositeLayerBase::IsEnabled() const
-{
-	return bIsEnabled;
-}
-
-void UCompositeLayerBase::SetEnabled(bool bInEnabled)
+bool UCompositeLayerBase::GetProxy(
+	const UE::CompositeCore::FPassInputDecl& /*InputDecl*/,
+	FCompositeTraversalContext& /*InContext*/,
+	FSceneRenderingBulkObjectAllocator& /*InFrameAllocator*/,
+	FCompositeCorePassProxy*& /*OutProxy*/) const
 {
-	bIsEnabled = bInEnabled;
+	return false;
 }
 
 void UCompositeLayerBase::UpdateSceneCaptureComponents(USceneCaptureComponent2D& SceneCaptureComponent, TArrayView<TSoftObjectPtr<AActor>> InActors)
diff --git a/Engine/Plugins/Compositing/Composite/Source/Composite/Private/Layers/CompositeLayerMainRender.cpp b/Engine/Plugins/Compositing/Composite/Source/Composite/Private/Layers/CompositeLayerMainRender.cpp
--- a/Engine/Plugins/Compositing/Composite/Source/Composite/Private/Layers/CompositeLayerMainRender.cpp
+++ b/Engine/Plugins/Compositing/Composite/Source/Composite/Private/Layers/CompositeLayerMainRender.cpp
@@ -12,7 +12,7 @@ UCompositeLayerMainRender::UCompositeLayerMainRender(const FObjectInitializer& O
 
 UCompositeLayerMainRender::~UCompositeLayerMainRender() = default;
 
-bool UCompositeLayerMainRender::GetProxy(FCompositeTraversalContext& InContext, FSceneRenderingBulkObjectAllocator& InFrameAllocator, FCompositeCorePassProxy*& OutProxy) const
+bool UCompositeLayerMainRender::GetProxy(const UE::CompositeCore::FPassInputDecl& /*InputDecl*/, FCompositeTraversalContext& InContext, FSceneRenderingBulkObjectAllocator& InFrameAllocator, FCompositeCorePassProxy*& OutProxy) const
 {
 	using namespace UE::CompositeCore;
 
diff --git a/Engine/Plugins/Compositing/Composite/Source/Composite/Private/Layers/CompositeLayerPlate.cpp b/Engine/Plugins/Compositing/Composite/Source/Composite/Private/Layers/CompositeLayerPlate.cpp
--- a/Engine/Plugins/Compositing/Composite/Source/Composite/Private/Layers/CompositeLayerPlate.cpp
+++ b/Engine/Plugins/Compositing/Composite/Source/Composite/Private/Layers/CompositeLayerPlate.cpp
@@ -527,7 +527,7 @@ UE::CompositeCore::ResourceId UCompositeLayerPlate::AddPreprocessingPasses(
 	return OutputTextureId;
 }
 
-bool UCompositeLayerPlate::GetProxy(FCompositeTraversalContext& InContext, FSceneRenderingBulkObjectAllocator& InFrameAllocator, FCompositeCorePassProxy*& OutProxy) const
+bool UCompositeLayerPlate::GetProxy(const UE::CompositeCore::FPassInputDecl& /*InputDecl*/, FCompositeTraversalContext& InContext, FSceneRenderingBulkObjectAllocator& InFrameAllocator, FCompositeCorePassProxy*& OutProxy) const
 {
 	using namespace UE::CompositeCore;
 	using namespace UE::Composite::Private;
diff --git a/Engine/Plugins/Compositing/Composite/Source/Composite/Private/Layers/CompositeLayerProcessing.cpp b/Engine/Plugins/Compositing/Composite/Source/Composite/Private/Layers/CompositeLayerProcessing.cpp
--- a/Engine/Plugins/Compositing/Composite/Source/Composite/Private/Layers/CompositeLayerProcessing.cpp
+++ b/Engine/Plugins/Compositing/Composite/Source/Composite/Private/Layers/CompositeLayerProcessing.cpp
@@ -40,10 +40,11 @@ void UCompositeLayerProcessing::PostEditChangeProperty(FPropertyChangedEvent& Pr
 }
 #endif
 
-bool UCompositeLayerProcessing::GetProxy(FCompositeTraversalContext& InContext, FSceneRenderingBulkObjectAllocator& InFrameAllocator, FCompositeCorePassProxy*& OutProxy) const
+bool UCompositeLayerProcessing::GetProxy(const UE::CompositeCore::FPassInputDecl& /*InputDecl*/, FCompositeTraversalContext& InContext, FSceneRenderingBulkObjectAllocator& InFrameAllocator, FCompositeCorePassProxy*& OutProxy) const
 {
 	using namespace UE::CompositeCore;
-
+	
+	// Note: We ignore the passed InputDecl for now, this will be cleaned up in a subsequence refactor.
 	FPassInputDecl BasePassInput = GetDefaultSecondInput(InContext);
 	OutProxy = nullptr;
 
diff --git a/Engine/Plugins/Compositing/Composite/Source/Composite/Private/Layers/CompositeLayerSceneCapture.cpp b/Engine/Plugins/Compositing/Composite/Source/Composite/Private/Layers/CompositeLayerSceneCapture.cpp
--- a/Engine/Plugins/Compositing/Composite/Source/Composite/Private/Layers/CompositeLayerSceneCapture.cpp
+++ b/Engine/Plugins/Compositing/Composite/Source/Composite/Private/Layers/CompositeLayerSceneCapture.cpp
@@ -138,7 +138,7 @@ void UCompositeLayerSceneCapture::OnEndOfFrameUpdate(UWorld* InWorld)
 	}
 }
 
-bool UCompositeLayerSceneCapture::GetProxy(FCompositeTraversalContext& InContext, FSceneRenderingBulkObjectAllocator& InFrameAllocator, FCompositeCorePassProxy*& OutProxy) const
+bool UCompositeLayerSceneCapture::GetProxy(const UE::CompositeCore::FPassInputDecl& /*InputDecl*/, FCompositeTraversalContext& InContext, FSceneRenderingBulkObjectAllocator& InFrameAllocator, FCompositeCorePassProxy*& OutProxy) const
 {
 	using namespace UE::CompositeCore;
 
diff --git a/Engine/Plugins/Compositing/Composite/Source/Composite/Private/Layers/CompositeLayerShadowReflection.cpp b/Engine/Plugins/Compositing/Composite/Source/Composite/Private/Layers/CompositeLayerShadowReflection.cpp
--- a/Engine/Plugins/Compositing/Composite/Source/Composite/Private/Layers/CompositeLayerShadowReflection.cpp
+++ b/Engine/Plugins/Compositing/Composite/Source/Composite/Private/Layers/CompositeLayerShadowReflection.cpp
@@ -141,7 +141,7 @@ void UCompositeLayerShadowReflection::OnEndOfFrameUpdate(UWorld* InWorld)
 	}
 }
 
-bool UCompositeLayerShadowReflection::GetProxy(FCompositeTraversalContext& InContext, FSceneRenderingBulkObjectAllocator& InFrameAllocator, FCompositeCorePassProxy*& OutProxy) const
+bool UCompositeLayerShadowReflection::GetProxy(const UE::CompositeCore::FPassInputDecl& /*InputDecl*/, FCompositeTraversalContext& InContext, FSceneRenderingBulkObjectAllocator& InFrameAllocator, FCompositeCorePassProxy*& OutProxy) const
 {
 	using namespace UE::CompositeCore;
 
diff --git a/Engine/Plugins/Compositing/Composite/Source/Composite/Private/Layers/CompositeLayerSingleLightShadow.cpp b/Engine/Plugins/Compositing/Composite/Source/Composite/Private/Layers/CompositeLayerSingleLightShadow.cpp
--- a/Engine/Plugins/Compositing/Composite/Source/Composite/Private/Layers/CompositeLayerSingleLightShadow.cpp
+++ b/Engine/Plugins/Compositing/Composite/Source/Composite/Private/Layers/CompositeLayerSingleLightShadow.cpp
@@ -226,7 +226,7 @@ void UCompositeLayerSingleLightShadow::OnEndOfFrameUpdate(UWorld* InWorld)
 	}
 }
 
-bool UCompositeLayerSingleLightShadow::GetProxy(FCompositeTraversalContext& InContext, FSceneRenderingBulkObjectAllocator& InFrameAllocator, FCompositeCorePassProxy*& OutProxy) const
+bool UCompositeLayerSingleLightShadow::GetProxy(const UE::CompositeCore::FPassInputDecl& /*InputDecl*/, FCompositeTraversalContext& InContext, FSceneRenderingBulkObjectAllocator& InFrameAllocator, FCompositeCorePassProxy*& OutProxy) const
 {
 	using namespace UE::CompositeCore;
 
diff --git a/Engine/Plugins/Compositing/Composite/Source/Composite/Public/Layers/CompositeLayerBase.h b/Engine/Plugins/Compositing/Composite/Source/Composite/Public/Layers/CompositeLayerBase.h
--- a/Engine/Plugins/Compositing/Composite/Source/Composite/Public/Layers/CompositeLayerBase.h
+++ b/Engine/Plugins/Compositing/Composite/Source/Composite/Public/Layers/CompositeLayerBase.h
@@ -4,7 +4,6 @@
 
 #include "CoreTypes.h"
 #include "Containers/ContainersFwd.h"
-#include "UObject/Object.h"
 
 #include "Passes/CompositePassBase.h"
 #include "Passes/CompositeCorePassMergeProxy.h"
@@ -16,9 +15,9 @@
 class ACompositeActor;
 class UWorld;
 
-/** Base class for defining a pass object, and its associated render proxy. */
+/** Base class for defining a composite layer object, and its associated render proxy. */
 UCLASS(MinimalAPI, Abstract, EditInlineNew)
-class UCompositeLayerBase : public UObject
+class UCompositeLayerBase : public UCompositePassBase
 {
 	GENERATED_BODY()
 
@@ -30,10 +29,11 @@ public:
 	UE_API virtual ~UCompositeLayerBase();
 
 	//~ Begin UObject Interface
+	UE_API virtual void PostLoad() override;
 	UE_API virtual void BeginDestroy() override;
 
 #if WITH_EDITOR
-	UE_API virtual bool CanEditChange(const FProperty* InProperty) const;
+	UE_API virtual bool CanEditChange(const FProperty* InProperty) const override;
 	UE_API virtual void PostEditChangeProperty(FPropertyChangedEvent& PropertyChangedEvent) override;
 #endif
 	//~ End UObject Interface
@@ -58,16 +58,8 @@ public:
 	*/
 	UE_API virtual void OnRemoved(ACompositeActor* LastOwner) {}
 
-	/** Getter function to override, returning pass layer proxies to be passed to the render thread. (Proxy objects should be allocated from the provided allocator.) */
-	UE_API virtual bool GetProxy(FCompositeTraversalContext& InContext, FSceneRenderingBulkObjectAllocator& InFrameAllocator, FCompositeCorePassProxy*& OutProxy) const { return false; };
-
-	/** Get the enabled state. */
-	UFUNCTION(BlueprintGetter)
-	UE_API virtual bool IsEnabled() const;
-
-	/** Set the enabled state. */
-	UFUNCTION(BlueprintSetter)
-	UE_API virtual void SetEnabled(bool bInEnabled);
+	/** Getter function to override, returning layer proxies to be passed to the render thread. (Proxy objects should be allocated from the provided allocator.) */
+	UE_API virtual bool GetProxy(const UE::CompositeCore::FPassInputDecl& InputDecl, FCompositeTraversalContext& InContext, FSceneRenderingBulkObjectAllocator& InFrameAllocator, FCompositeCorePassProxy*& OutProxy) const override;
 
 protected:
 	/** Convenience function to update scene capture primitive components, depending on the scene capture primitive render mode. */
@@ -82,20 +74,15 @@ public:
 	bool bIsSolo = false;
 
 #if WITH_EDITORONLY_DATA
-	/** Pass layer name. */
-	UPROPERTY(EditAnywhere, BlueprintReadWrite, Category = "", meta = (DisplayPriority = "1"))
-	FString Name;
+	/** Deprecated layer display name, use UCompositePassBase::DisplayName. */
+	UPROPERTY(meta = (DeprecatedProperty, DeprecationMessage = "Use DisplayName instead."))
+	FString Name_DEPRECATED;
 #endif
 
 	/** Merge operation applied on Input0 with Input1. */
 	UPROPERTY(EditAnywhere, BlueprintReadWrite, Category = "", meta = (DisplayPriority = "3"))
 	ECompositeCoreMergeOp Operation = ECompositeCoreMergeOp::Over;
 
-protected:
-	/** Whether or not the pass is enabled. */
-	UPROPERTY(EditAnywhere, BlueprintGetter = IsEnabled, BlueprintSetter = SetEnabled, Category = "", meta = (DisplayPriority = "2"))
-	bool bIsEnabled = true;
-
 	/** Convenience function to register post-processing child passes. */
 	void AddChildPasses(UE::CompositeCore::FPassInputDecl& InBasePassInput, FCompositeTraversalContext& InContext, FSceneRenderingBulkObjectAllocator& InFrameAllocator, TArrayView<const TObjectPtr<UCompositePassBase>> InPasses) const;
 
diff --git a/Engine/Plugins/Compositing/Composite/Source/Composite/Public/Layers/CompositeLayerMainRender.h b/Engine/Plugins/Compositing/Composite/Source/Composite/Public/Layers/CompositeLayerMainRender.h
--- a/Engine/Plugins/Compositing/Composite/Source/Composite/Public/Layers/CompositeLayerMainRender.h
+++ b/Engine/Plugins/Compositing/Composite/Source/Composite/Public/Layers/CompositeLayerMainRender.h
@@ -24,7 +24,7 @@ public:
 	UE_API virtual ~UCompositeLayerMainRender();
 
 	/** Getter function to override, returning pass layer proxies to be passed to the render thread. (Proxy objects should be allocated from the provided allocator.) */
-	UE_API virtual bool GetProxy(FCompositeTraversalContext& InContext, FSceneRenderingBulkObjectAllocator& InFrameAllocator, FCompositeCorePassProxy*& OutProxy) const override;
+	UE_API virtual bool GetProxy(const UE::CompositeCore::FPassInputDecl& InputDecl, FCompositeTraversalContext& InContext, FSceneRenderingBulkObjectAllocator& InFrameAllocator, FCompositeCorePassProxy*& OutProxy) const override;
 };
 
 #undef UE_API
diff --git a/Engine/Plugins/Compositing/Composite/Source/Composite/Public/Layers/CompositeLayerPlate.h b/Engine/Plugins/Compositing/Composite/Source/Composite/Public/Layers/CompositeLayerPlate.h
--- a/Engine/Plugins/Compositing/Composite/Source/Composite/Public/Layers/CompositeLayerPlate.h
+++ b/Engine/Plugins/Compositing/Composite/Source/Composite/Public/Layers/CompositeLayerPlate.h
@@ -69,7 +69,7 @@ public:
 	//~ End FTickableGameObject interface
 
 	/** Getter function to override, returning pass layer proxies to be passed to the render thread. (Proxy objects should be allocated from the provided allocator.) */
-	UE_API virtual bool GetProxy(FCompositeTraversalContext& InContext, FSceneRenderingBulkObjectAllocator& InFrameAllocator, FCompositeCorePassProxy*& OutProxy) const override;
+	UE_API virtual bool GetProxy(const UE::CompositeCore::FPassInputDecl& InputDecl, FCompositeTraversalContext& InContext, FSceneRenderingBulkObjectAllocator& InFrameAllocator, FCompositeCorePassProxy*& OutProxy) const override;
 
 	/** Get array of all composite mesh primitives. */
 	UE_API TArray<UPrimitiveComponent*> GetPrimitives() const;
@@ -132,7 +132,7 @@ public:
 	* 
 	* The texture is replaced with an internal render target.
 	*/
-	UPROPERTY(EditAnywhere, BlueprintReadWrite, Instanced, Category = "", meta = (DisallowedClasses = "/Script/Composite.CompositePassMaterial, /Script/Composite.CompositePassDistortion, /Script/Composite.CompositePassFXAA, /Script/Composite.CompositePassSMAA"))
+	UPROPERTY(EditAnywhere, BlueprintReadWrite, Instanced, Category = "", meta = (DisallowedClasses = "/Script/Composite.CompositeLayerBase, /Script/Composite.CompositePassMaterial, /Script/Composite.CompositePassDistortion, /Script/Composite.CompositePassFXAA, /Script/Composite.CompositePassSMAA"))
 	TArray<TObjectPtr<UCompositePassBase>> MediaPasses;
 
 	/**
@@ -144,11 +144,11 @@ public:
 	* We now have an automatic scene pass to undistort a plate projected onto composite meshes. Using Texture mode can therefore
 	* avoid the automatic undistort -> distort pipeline.
 	*/
-	UPROPERTY(BlueprintReadWrite, Instanced, Category = "", meta = (AllowedClasses = "/Script/Composite.CompositePassDistortion"))
+	UPROPERTY(BlueprintReadWrite, Instanced, Category = "", meta = (AllowedClasses = "/Script/Composite.CompositePassDistortion", DisallowedClasses = "/Script/Composite.CompositeLayerBase"))
 	TArray<TObjectPtr<UCompositePassBase>> ScenePasses;
 
 	/** Passes applied during final compositing in post-processing. */
-	UPROPERTY(EditAnywhere, BlueprintReadWrite, Instanced, Category = "", meta = (DisallowedClasses = "/Script/Composite.CompositePassDistortion"))
+	UPROPERTY(EditAnywhere, BlueprintReadWrite, Instanced, Category = "", meta = (DisallowedClasses = "/Script/Composite.CompositeLayerBase, /Script/Composite.CompositePassDistortion"))
 	TArray<TObjectPtr<UCompositePassBase>> LayerPasses;
 
 	/** Composite.Plate.ControlMeshVisibility console variable value. */
diff --git a/Engine/Plugins/Compositing/Composite/Source/Composite/Public/Layers/CompositeLayerProcessing.h b/Engine/Plugins/Compositing/Composite/Source/Composite/Public/Layers/CompositeLayerProcessing.h
--- a/Engine/Plugins/Compositing/Composite/Source/Composite/Public/Layers/CompositeLayerProcessing.h
+++ b/Engine/Plugins/Compositing/Composite/Source/Composite/Public/Layers/CompositeLayerProcessing.h
@@ -24,7 +24,7 @@ public:
 	UE_API virtual ~UCompositeLayerProcessing();
 
 	/** Getter function to override, returning pass layer proxies to be passed to the render thread. (Proxy objects should be allocated from the provided allocator.) */
-	UE_API virtual bool GetProxy(FCompositeTraversalContext& InContext, FSceneRenderingBulkObjectAllocator& InFrameAllocator, FCompositeCorePassProxy*& OutProxy) const override;
+	UE_API virtual bool GetProxy(const UE::CompositeCore::FPassInputDecl& InputDecl, FCompositeTraversalContext& InContext, FSceneRenderingBulkObjectAllocator& InFrameAllocator, FCompositeCorePassProxy*& OutProxy) const override;
 
 #if WITH_EDITOR
 	UE_API virtual bool CanEditChange(const FProperty* InProperty) const;
@@ -32,7 +32,7 @@ public:
 #endif
 
 	/** Sub-passes applied onto the previous layer. */
-	UPROPERTY(EditAnywhere, BlueprintReadWrite, Instanced, Category = "", meta = (DisallowedClasses = "/Script/Composite.CompositePassDistortion"))
+	UPROPERTY(EditAnywhere, BlueprintReadWrite, Instanced, Category = "", meta = (DisallowedClasses = "/Script/Composite.CompositeLayerBase, /Script/Composite.CompositePassDistortion"))
 	TArray<TObjectPtr<UCompositePassBase>> LayerPasses;
 };
 
diff --git a/Engine/Plugins/Compositing/Composite/Source/Composite/Public/Layers/CompositeLayerSceneCapture.h b/Engine/Plugins/Compositing/Composite/Source/Composite/Public/Layers/CompositeLayerSceneCapture.h
--- a/Engine/Plugins/Compositing/Composite/Source/Composite/Public/Layers/CompositeLayerSceneCapture.h
+++ b/Engine/Plugins/Compositing/Composite/Source/Composite/Public/Layers/CompositeLayerSceneCapture.h
@@ -45,7 +45,7 @@ public:
 	UE_API virtual void OnEndOfFrameUpdate(UWorld* InWorld) override;
 
 	/** Getter function to override, returning pass layer proxies to be passed to the render thread. (Proxy objects should be allocated from the provided allocator.) */
-	UE_API virtual bool GetProxy(FCompositeTraversalContext& InContext, FSceneRenderingBulkObjectAllocator& InFrameAllocator, FCompositeCorePassProxy*& OutProxy) const override;
+	UE_API virtual bool GetProxy(const UE::CompositeCore::FPassInputDecl& InputDecl, FCompositeTraversalContext& InContext, FSceneRenderingBulkObjectAllocator& InFrameAllocator, FCompositeCorePassProxy*& OutProxy) const override;
 
 	/** Get reflection & shadow caster actors. */
 	UFUNCTION(BlueprintGetter)
@@ -80,7 +80,7 @@ public:
 	UE_API void SetRenderTargetResolution(FIntPoint InRenderTargetResolution);
 
 	/** Sub-passes applied during post-processing. */
-	UPROPERTY(EditAnywhere, BlueprintReadWrite, Instanced, Category = "", meta = (DisallowedClasses = "/Script/Composite.CompositePassDistortion", DisplayAfter = "bCustomRenderPass"))
+	UPROPERTY(EditAnywhere, BlueprintReadWrite, Instanced, Category = "", meta = (DisallowedClasses = "/Script/Composite.CompositeLayerBase, /Script/Composite.CompositePassDistortion", DisplayAfter = "bCustomRenderPass"))
 	TArray<TObjectPtr<UCompositePassBase>> LayerPasses;
 
 private:
diff --git a/Engine/Plugins/Compositing/Composite/Source/Composite/Public/Layers/CompositeLayerShadowReflection.h b/Engine/Plugins/Compositing/Composite/Source/Composite/Public/Layers/CompositeLayerShadowReflection.h
--- a/Engine/Plugins/Compositing/Composite/Source/Composite/Public/Layers/CompositeLayerShadowReflection.h
+++ b/Engine/Plugins/Compositing/Composite/Source/Composite/Public/Layers/CompositeLayerShadowReflection.h
@@ -79,7 +79,7 @@ public:
 	UE_API virtual void OnEndOfFrameUpdate(UWorld* InWorld) override;
 
 	/** Getter function to override, returning pass layer proxies to be passed to the render thread. (Proxy objects should be allocated from the provided allocator.) */
-	UE_API virtual bool GetProxy(FCompositeTraversalContext& InContext, FSceneRenderingBulkObjectAllocator& InFrameAllocator, FCompositeCorePassProxy*& OutProxy) const override;
+	UE_API virtual bool GetProxy(const UE::CompositeCore::FPassInputDecl& InputDecl, FCompositeTraversalContext& InContext, FSceneRenderingBulkObjectAllocator& InFrameAllocator, FCompositeCorePassProxy*& OutProxy) const override;
 
 	/** Get the render target resolution for the scene captures. */
 	UFUNCTION(BlueprintGetter)
diff --git a/Engine/Plugins/Compositing/Composite/Source/Composite/Public/Layers/CompositeLayerSingleLightShadow.h b/Engine/Plugins/Compositing/Composite/Source/Composite/Public/Layers/CompositeLayerSingleLightShadow.h
--- a/Engine/Plugins/Compositing/Composite/Source/Composite/Public/Layers/CompositeLayerSingleLightShadow.h
+++ b/Engine/Plugins/Compositing/Composite/Source/Composite/Public/Layers/CompositeLayerSingleLightShadow.h
@@ -55,7 +55,7 @@ public:
 	UE_API virtual void OnEndOfFrameUpdate(UWorld* InWorld) override;
 
 	/** Getter function to override, returning pass layer proxies to be passed to the render thread. (Proxy objects should be allocated from the provided allocator.) */
-	UE_API virtual bool GetProxy(FCompositeTraversalContext& InContext, FSceneRenderingBulkObjectAllocator& InFrameAllocator, FCompositeCorePassProxy*& OutProxy) const override;
+	UE_API virtual bool GetProxy(const UE::CompositeCore::FPassInputDecl& InputDecl, FCompositeTraversalContext& InContext, FSceneRenderingBulkObjectAllocator& InFrameAllocator, FCompositeCorePassProxy*& OutProxy) const override;
 
 public:
 	/**
diff --git a/Engine/Plugins/Compositing/Composite/Source/CompositeEditor/Private/ActorFactories/CompositeActorFactory.cpp b/Engine/Plugins/Compositing/Composite/Source/CompositeEditor/Private/ActorFactories/CompositeActorFactory.cpp
--- a/Engine/Plugins/Compositing/Composite/Source/CompositeEditor/Private/ActorFactories/CompositeActorFactory.cpp
+++ b/Engine/Plugins/Compositing/Composite/Source/CompositeEditor/Private/ActorFactories/CompositeActorFactory.cpp
@@ -35,10 +35,10 @@ void UCompositeActorFactory::PostSpawnActor(UObject* Asset, AActor* NewActor)
 		NewLayers.Add(NewObject<UCompositeLayerShadowReflection>(CompositeActor, TEXT("ShadowReflectionLayer"), RF_Transactional));
 		NewLayers.Add(NewObject<UCompositeLayerPlate>(CompositeActor, TEXT("PlateLayer"), RF_Transactional));
 
-		// Assign default Name value to all layers
+		// Assign default display name value to all layers
 		for (TObjectPtr<UCompositeLayerBase>& Layer : NewLayers)
 		{
-			Layer->Name = Layer->GetName();
+			Layer->SetDisplayName(Layer->GetName());
 		}
 
 		// By default, we disable the shadow reflection catcher due to its significant added cost.
diff --git a/Engine/Plugins/Compositing/Composite/Source/CompositeEditor/Private/ColorGradingDrawer/ColorGradingHierarchyConfig_Composite.cpp b/Engine/Plugins/Compositing/Composite/Source/CompositeEditor/Private/ColorGradingDrawer/ColorGradingHierarchyConfig_Composite.cpp
--- a/Engine/Plugins/Compositing/Composite/Source/CompositeEditor/Private/ColorGradingDrawer/ColorGradingHierarchyConfig_Composite.cpp
+++ b/Engine/Plugins/Compositing/Composite/Source/CompositeEditor/Private/ColorGradingDrawer/ColorGradingHierarchyConfig_Composite.cpp
@@ -21,7 +21,7 @@ TArray<FObjectMixerEditorActorSubObject> FColorGradingHierarchyConfig_Composite:
 			{
 				FObjectMixerEditorActorSubObject SubObject;
 				SubObject.Object = CompositeLayer;
-				SubObject.DisplayName = FText::FromString(!CompositeLayer->Name.IsEmpty() ? CompositeLayer->Name : CompositeLayer->GetName());
+				SubObject.DisplayName = FText::FromString(!CompositeLayer->GetDisplayName().IsEmpty() ? CompositeLayer->GetDisplayName() : CompositeLayer->GetName());
 				SubObject.bIncludeObjectStack = true;
 				
 				SubObjects.Add(SubObject);
diff --git a/Engine/Plugins/Compositing/Composite/Source/CompositeEditor/Private/UI/SCompositeEditorPanel.cpp b/Engine/Plugins/Compositing/Composite/Source/CompositeEditor/Private/UI/SCompositeEditorPanel.cpp
--- a/Engine/Plugins/Compositing/Composite/Source/CompositeEditor/Private/UI/SCompositeEditorPanel.cpp
+++ b/Engine/Plugins/Compositing/Composite/Source/CompositeEditor/Private/UI/SCompositeEditorPanel.cpp
@@ -235,7 +235,8 @@ void SCompositeEditorPanel::SelectObjects(const TArray<UObject*>& InSelection)
 			CompositeActorsAndLayers.Add(Object);
 		}
 
-		if (UCompositePassBase* Pass = Cast<UCompositePassBase>(Object))
+		UCompositePassBase* Pass = Cast<UCompositePassBase>(Object);
+		if (Pass && !Pass->IsA<UCompositeLayerBase>())
 		{
 			if (UCompositeLayerBase* ParentLayer = Pass->GetTypedOuter<UCompositeLayerBase>())
 			{
diff --git a/Engine/Plugins/Compositing/Composite/Source/CompositeEditor/Private/UI/SCompositePanelLayerTree.cpp b/Engine/Plugins/Compositing/Composite/Source/CompositeEditor/Private/UI/SCompositePanelLayerTree.cpp
--- a/Engine/Plugins/Compositing/Composite/Source/CompositeEditor/Private/UI/SCompositePanelLayerTree.cpp
+++ b/Engine/Plugins/Compositing/Composite/Source/CompositeEditor/Private/UI/SCompositePanelLayerTree.cpp
@@ -179,7 +179,7 @@ private:
 				
 			if (UCompositeLayerBase* CompositeLayer = InItem->GetCompositeLayer())
 			{
-				OutStrings.Add(CompositeLayer->Name);
+				OutStrings.Add(CompositeLayer->GetDisplayName());
 				OutStrings.Add(CompositeLayer->GetClass()->GetDisplayNameText().ToString());
 			}
 		}));
@@ -336,7 +336,7 @@ private:
 			{
 				if (UCompositeLayerBase* CompositeLayer = PinnedItem->GetCompositeLayer())
 				{
-					const FString LayerName = !CompositeLayer->Name.IsEmpty() ? CompositeLayer->Name : CompositeLayer->GetName();
+					const FString LayerName = !CompositeLayer->GetDisplayName().IsEmpty() ? CompositeLayer->GetDisplayName() : CompositeLayer->GetName();
 					const FSlateBrush* LayerIcon = FSlateIconFinder::FindIconForClass(CompositeLayer->GetClass()).GetIcon();
 					SetToolTip(FText::FromString(LayerName), LayerIcon);
 				}
@@ -655,13 +655,13 @@ private:
 		{
 			if (UCompositeLayerBase* CompositeLayer = TreeItem->GetCompositeLayer())
 			{
-				if (CompositeLayer->Name.IsEmpty())
+				if (CompositeLayer->GetDisplayName().IsEmpty())
 				{
-					// Display the UObject name for the layer if its Name property is empty.
+					// Display the UObject name for the layer if its display name is empty.
 					return FText::FromString(CompositeLayer->GetName());
 				}
 			
-				return FText::FromString(CompositeLayer->Name);
+				return FText::FromString(CompositeLayer->GetDisplayName());
 			}
 
 			// In the case where there is a valid layer index, but the layer itself is null
@@ -686,7 +686,7 @@ private:
 				FScopedTransaction SetLayerNameTransaction(LOCTEXT("SetLayerNameTransaction", "Set Layer Name"));
 				CompositeLayer->Modify();
 
-				CompositeLayer->Name = InNewName.ToString();
+				CompositeLayer->SetDisplayName(InNewName.ToString());
 				return;
 			}
 
diff --git a/Engine/Plugins/Compositing/Composite/Source/CompositeEditor/Private/UI/SCompositePassTree.cpp b/Engine/Plugins/Compositing/Composite/Source/CompositeEditor/Private/UI/SCompositePassTree.cpp
--- a/Engine/Plugins/Compositing/Composite/Source/CompositeEditor/Private/UI/SCompositePassTree.cpp
+++ b/Engine/Plugins/Compositing/Composite/Source/CompositeEditor/Private/UI/SCompositePassTree.cpp
@@ -340,44 +340,45 @@ private:
 		constexpr bool bCloseMenuAfterSelection = true;
 		FMenuBuilder MenuBuilder(bCloseMenuAfterSelection, nullptr);
 
-		TArray<UClass*> BaseLayerTypes;
+		TArray<UClass*> BasePassTypes;
 		const bool bRecursive = false;
-		GetDerivedClasses(UCompositePassBase::StaticClass(), BaseLayerTypes, bRecursive);
+		GetDerivedClasses(UCompositePassBase::StaticClass(), BasePassTypes, bRecursive);
 
 		auto ShouldSkipClass = [](const UClass* InClass)
 			{
 				constexpr EClassFlags InvalidClassFlags = CLASS_Deprecated | CLASS_NewerVersionExists | CLASS_Abstract;
 				return InClass->HasAnyClassFlags(InvalidClassFlags) ||
+					InClass->IsChildOf(UCompositeLayerBase::StaticClass()) ||
 					InClass->GetName().StartsWith(TEXT("SKEL_")) ||
 					InClass->GetName().StartsWith(TEXT("REINST_"));
 			};
 		
-		// Add the native child classes of UCompositeLayerBase to the menu first
-		for (const UClass* BaseLayerType : BaseLayerTypes)
+		// Add the native child classes of UCompositePassBase to the menu first
+		for (const UClass* BasePassType : BasePassTypes)
 		{
-			if (ShouldSkipClass(BaseLayerType))
+			if (ShouldSkipClass(BasePassType))
 			{
 				continue;
 			}
 
 			if (OnFilterNewPassType.IsBound())
 			{
-				if (OnFilterNewPassType.Execute(BaseLayerType))
+				if (OnFilterNewPassType.Execute(BasePassType))
 				{
 					continue;
 				}
 			}
 			
 			MenuBuilder.AddMenuEntry(
-				BaseLayerType->GetDisplayNameText(),
+				BasePassType->GetDisplayNameText(),
 				TAttribute<FText>(),
-				FSlateIconFinder::FindIconForClass(BaseLayerType),
-				FUIAction(FExecuteAction::CreateSP(this, &SCompositePassTreeToolbar::AddPass, BaseLayerType)));
+				FSlateIconFinder::FindIconForClass(BasePassType),
+				FUIAction(FExecuteAction::CreateSP(this, &SCompositePassTreeToolbar::AddPass, BasePassType)));
 		}
 				
 		// Add any blueprint derived classes in a new section of the menu
 		bool bHasAddedSeparator = false;
-		for (const UClass* BaseLayerType : BaseLayerTypes)
+		for (const UClass* BaseLayerType : BasePassTypes)
 		{
 			TArray<UClass*> ChildLayerTypes;
 			GetDerivedClasses(BaseLayerType, ChildLayerTypes);
@@ -391,7 +392,7 @@ private:
 
 				if (OnFilterNewPassType.IsBound())
 				{
-					if (OnFilterNewPassType.Execute(BaseLayerType))
+					if (OnFilterNewPassType.Execute(ChildLayerType))
 					{
 						continue;
 					}
@@ -1467,6 +1468,11 @@ void SCompositePassTree::OnFilterChanged()
 
 void SCompositePassTree::OnPassAdded(const UClass* InPassClass)
 {
+	if (!IsValid(InPassClass) || InPassClass->IsChildOf<UCompositeLayerBase>())
+	{
+		return;
+	}
+
 	TArray<FPassTreeItemPtr> SelectedItems = TreeView.IsValid() ? TreeView->GetSelectedItems() : TArray<FPassTreeItemPtr>();
 	if (SelectedItems.IsEmpty())
 	{
@@ -1525,6 +1531,11 @@ void SCompositePassTree::OnPassAdded(const UClass* InPassClass)
 
 bool SCompositePassTree::OnFilterNewPassType(const UClass* InPassType) const
 {
+	if (!IsValid(InPassType) || InPassType->IsChildOf<UCompositeLayerBase>())
+	{
+		return true;
+	}
+
 	// Filter out Distortion passes as those are implicitly added when there is a lens file on the composite actor camera
 	if (InPassType->IsChildOf<UCompositePassDistortion>())
 	{
@@ -1736,7 +1747,7 @@ public:
 
 	virtual bool CanCreateClass(UClass* ObjectClass, bool& bOmitSubObjs) const override
 	{
-		return ObjectClass->IsChildOf(UCompositePassBase::StaticClass());
+		return ObjectClass->IsChildOf(UCompositePassBase::StaticClass()) && !ObjectClass->IsChildOf(UCompositeLayerBase::StaticClass());
 	}
 
 	virtual void ProcessConstructedObject(UObject* NewObject) override
@@ -1745,7 +1756,10 @@ public:
 
 		if (UCompositePassBase* Pass = Cast<UCompositePassBase>(NewObject))
 		{
-			CompositePasses.Add(Pass);
+			if (!Pass->IsA<UCompositeLayerBase>())
+			{
+				CompositePasses.Add(Pass);
+			}
 		}
 	}
 
diff --git a/Engine/Source/Runtime/Core/Public/UObject/FortniteMainBranchObjectVersions.inl b/Engine/Source/Runtime/Core/Public/UObject/FortniteMainBranchObjectVersions.inl
--- a/Engine/Source/Runtime/Core/Public/UObject/FortniteMainBranchObjectVersions.inl
+++ b/Engine/Source/Runtime/Core/Public/UObject/FortniteMainBranchObjectVersions.inl
@@ -778,5 +778,8 @@ FORTNITE_MAIN_VERSION(AnimSequenceTransformProviderBounds, 254)
 // UEnum now have Underlyingtype
 FORTNITE_MAIN_VERSION(EnumUnderlyingType, 255)
 
+// Composite layer name deprecation
+FORTNITE_MAIN_VERSION(Composite_LayerNameDeprecation, 256)
+
 // -----<new versions can be added above this line>-------------------------------------------------
 #undef FORTNITE_MAIN_VERSION
