Composure: Added planar reflection layer, as a new (potentially faster) option for catching reflections.

* Can be used in combination with the single light shadow layer.
* Blur pass can be added to modulate the reflections.

#Synced-CL 51340812

diff --git a/Engine/Plugins/Compositing/Composite/Shaders/Private/CompositePlanarReflection.usf b/Engine/Plugins/Compositing/Composite/Shaders/Private/CompositePlanarReflection.usf
new file mode 100644
--- /dev/null
+++ b/Engine/Plugins/Compositing/Composite/Shaders/Private/CompositePlanarReflection.usf
@@ -0,0 +1,19 @@
+// Copyright Epic Games, Inc. All Rights Reserved.
+
+#include "/Engine/Private/Common.ush"
+
+Texture2D InputTexture;
+SamplerState InputSampler;
+float ReflectionStrength;
+
+void MainPS(
+	noperspective float4 InUVAndScreenPos : TEXCOORD0,
+	out float4 OutColor : SV_Target0
+	)
+{
+	float2 UV = float2(InUVAndScreenPos.x, 1.0f - InUVAndScreenPos.y);
+	const float4 InputColor = InputTexture.Sample(InputSampler, UV);
+	// Reflection blurring/filtering could be integrated here directly.
+
+	OutColor = lerp(0.0f, InputColor, ReflectionStrength);
+}
diff --git a/Engine/Plugins/Compositing/Composite/Source/Composite/Private/Layers/CompositeLayerPlanarReflection.cpp b/Engine/Plugins/Compositing/Composite/Source/Composite/Private/Layers/CompositeLayerPlanarReflection.cpp
new file mode 100644
--- /dev/null
+++ b/Engine/Plugins/Compositing/Composite/Source/Composite/Private/Layers/CompositeLayerPlanarReflection.cpp
@@ -0,0 +1,442 @@
+// Copyright Epic Games, Inc. All Rights Reserved.
+
+#include "Layers/CompositeLayerPlanarReflection.h"
+
+#include "CompositeActor.h"
+#include "Components/CompositeSceneCapture2DComponent.h"
+#include "CompositeRenderTargetPool.h"
+#include "Misc/TransactionObjectEvent.h"
+#include "Passes/CompositeCorePassProxy.h"
+#include "RenderGraphBuilder.h"
+#include "SceneView.h"
+#include "ScreenPass.h"
+#include "Camera/CameraComponent.h"
+
+#if WITH_EDITOR
+#include "Editor.h"
+#include "LevelEditorSubsystem.h"
+#include "LevelEditorViewport.h"
+#endif
+
+DECLARE_GPU_STAT_NAMED(FCompositePlanarReflection, TEXT("Composite.PlanarReflection"));
+
+class FCompositePlanarReflectionShader : public FGlobalShader
+{
+	DECLARE_GLOBAL_SHADER(FCompositePlanarReflectionShader);
+	SHADER_USE_PARAMETER_STRUCT(FCompositePlanarReflectionShader, FGlobalShader);
+
+	BEGIN_SHADER_PARAMETER_STRUCT(FParameters, )
+		SHADER_PARAMETER_RDG_TEXTURE(Texture2D, InputTexture)
+		SHADER_PARAMETER_SAMPLER(SamplerState, InputSampler)
+		SHADER_PARAMETER(float, ReflectionStrength)
+		RENDER_TARGET_BINDING_SLOTS()
+	END_SHADER_PARAMETER_STRUCT()
+};
+IMPLEMENT_GLOBAL_SHADER(FCompositePlanarReflectionShader, "/Plugin/Composite/Private/CompositePlanarReflection.usf", "MainPS", SF_Pixel);
+
+namespace UE
+{
+	namespace Composite
+	{
+		namespace Private
+		{
+			/** Resolve the planar reflection capture source from the composite actor output mode. */
+			ESceneCaptureSource GetPlanarReflectionCaptureSource(const ACompositeActor& CompositeActor, bool bCustomRenderPass)
+			{
+				if (bCustomRenderPass)
+				{
+					return ESceneCaptureSource::SCS_SceneColorHDR;
+				}
+
+				return CompositeActor.MainRenderOutput == ECompositeMainRenderOutputMode::FinalColorHDR
+					? ESceneCaptureSource::SCS_FinalColorHDR
+					: ESceneCaptureSource::SCS_FinalToneCurveHDR;
+			}
+
+#if WITH_EDITOR
+			/** Get the currently active level editor viewport client, if any. */
+			FLevelEditorViewportClient* GetActiveLevelEditorViewportClient()
+			{
+				if (const FViewport* ActiveViewport = GEditor ? GEditor->GetActiveViewport() : nullptr)
+				{
+					for (FLevelEditorViewportClient* LevelViewportClient : GEditor->GetLevelViewportClients())
+					{
+						if (LevelViewportClient && LevelViewportClient->Viewport == ActiveViewport)
+						{
+							return LevelViewportClient;
+						}
+					}
+				}
+
+				return nullptr;
+			}
+
+			/**
+			* Build a view from the active editor viewport when the composite camera is not being piloted.
+			* This is needed here because we cannot rely on the scene capture's "MainViewCamera" behavior.
+			*/
+			bool TryGetEditorViewportCameraView(ACompositeActor& CompositeActor, FMinimalViewInfo& OutViewInfo)
+			{
+				FLevelEditorViewportClient* ViewportClient = GetActiveLevelEditorViewportClient();
+				if (!ViewportClient || !ViewportClient->IsPerspective())
+				{
+					return false;
+				}
+
+				const UCameraComponent* CompositeCamera = Cast<UCameraComponent>(CompositeActor.GetCamera().GetComponent(nullptr));
+				const AActor* CompositeCameraActor = CompositeCamera ? CompositeCamera->GetOwner() : nullptr;
+
+				if (ULevelEditorSubsystem* LevelEditorSubsystem = GEditor->GetEditorSubsystem<ULevelEditorSubsystem>())
+				{
+					if (CompositeCameraActor && LevelEditorSubsystem->GetPilotLevelActor() == CompositeCameraActor)
+					{
+						return false;
+					}
+				}
+
+				OutViewInfo.Location = ViewportClient->GetViewLocation();
+				OutViewInfo.Rotation = ViewportClient->GetViewRotation();
+				OutViewInfo.FOV = ViewportClient->ViewFOV;
+				OutViewInfo.ProjectionMode = ECameraProjectionMode::Perspective;
+				OutViewInfo.bConstrainAspectRatio = false;
+
+				if (ViewportClient->Viewport)
+				{
+					const FIntPoint ViewportSize = ViewportClient->Viewport->GetSizeXY();
+					OutViewInfo.AspectRatio = FMath::SafeDivide(static_cast<float>(ViewportSize.X), static_cast<float>(ViewportSize.Y));
+				}
+
+				return true;
+			}
+#endif //WITH_EDITOR
+		}
+	}
+}
+
+namespace UE
+{
+	namespace CompositeCore
+	{
+		class FPlanarReflectionPassProxy : public FCompositeCorePassProxy
+		{
+		public:
+			IMPLEMENT_COMPOSITE_PASS(FPlanarReflectionPassProxy);
+
+			using FCompositeCorePassProxy::FCompositeCorePassProxy;
+
+			FPassTexture Add(FRDGBuilder& GraphBuilder, const FSceneView& InView, const FPassInputArray& Inputs, const FPassContext& PassContext) const override
+			{
+				RDG_EVENT_SCOPE_STAT(GraphBuilder, FCompositePlanarReflection, "Composite.PlanarReflection");
+
+				check(ValidateInputs(Inputs));
+
+				const FResourceMetadata& Metadata = Inputs[0].Metadata;
+				const FScreenPassTexture& Input = Inputs[0].Texture;
+				FScreenPassRenderTarget Output = Inputs.OverrideOutput;
+
+				if (!Output.IsValid())
+				{
+					Output = CreateOutputRenderTarget(GraphBuilder, InView, PassContext.OutputViewRect, Input.Texture->Desc, TEXT("CompositePlanarReflection"));
+				}
+
+				FGlobalShaderMap* GlobalShaderMap = GetGlobalShaderMap(InView.GetFeatureLevel());
+				const FScreenPassTextureViewport Viewport(Output);
+
+				FCompositePlanarReflectionShader::FParameters* Parameters = GraphBuilder.AllocParameters<FCompositePlanarReflectionShader::FParameters>();
+				Parameters->InputTexture = Input.Texture;
+				Parameters->InputSampler = TStaticSamplerState<SF_Bilinear>::GetRHI();
+				Parameters->ReflectionStrength = ReflectionStrength;
+				Parameters->RenderTargets[0] = Output.GetRenderTargetBinding();
+
+				TShaderMapRef<FCompositePlanarReflectionShader> PixelShader(GlobalShaderMap);
+				AddDrawScreenPass(
+					GraphBuilder,
+					RDG_EVENT_NAME("Composite.PlanarReflection (%dx%d)", Viewport.Extent.X, Viewport.Extent.Y),
+					InView,
+					Viewport,
+					Viewport,
+					PixelShader,
+					Parameters
+				);
+
+				return FPassTexture{ MoveTemp(Output), Metadata };
+			}
+
+			float ReflectionStrength = 0.5f;
+		};
+	}
+}
+
+UCompositeLayerPlanarReflection::UCompositeLayerPlanarReflection(const FObjectInitializer& ObjectInitializer)
+	: Super(ObjectInitializer)
+	, RenderTargetResolution{ FCompositeRenderTargetPool::DefaultSize / 2 }
+	, ReflectionStrength{ 0.5f }
+	, PlaneHeight{ 0.0f }
+	, bCustomRenderPass{ false }
+{
+	Operation = ECompositeCoreMergeOp::Screen;
+
+	if (!HasAnyFlags(RF_ClassDefaultObject))
+	{
+		constexpr bool bEnabled = true;
+		RegisterEndOfFrameUpdate(bEnabled);
+	}
+}
+
+UCompositeLayerPlanarReflection::~UCompositeLayerPlanarReflection() = default;
+
+void UCompositeLayerPlanarReflection::OnRemoved(ACompositeActor* LastOwner)
+{
+	if (LastOwner != nullptr)
+	{
+		LastOwner->DestroySceneCaptures(this);
+	}
+}
+
+void UCompositeLayerPlanarReflection::BeginDestroy()
+{
+	Super::BeginDestroy();
+
+	OnRemoved(GetTypedOuter<ACompositeActor>()); // Redundant remove call for safety
+}
+
+#if WITH_EDITOR
+void UCompositeLayerPlanarReflection::PostEditChangeProperty(FPropertyChangedEvent& PropertyChangedEvent)
+{
+	Super::PostEditChangeProperty(PropertyChangedEvent);
+
+	const FName PropertyName = PropertyChangedEvent.GetMemberPropertyName();
+
+	if (PropertyName == GET_MEMBER_NAME_CHECKED(ThisClass, Actors))
+	{
+		SetActors(MoveTemp(Actors));
+	}
+	else if (PropertyName == GET_MEMBER_NAME_CHECKED(ThisClass, bCustomRenderPass))
+	{
+		SetCustomRenderPass(bCustomRenderPass);
+	}
+	else if (PropertyName == GET_MEMBER_NAME_CHECKED(ThisClass, RenderTargetResolution))
+	{
+		SetRenderTargetResolution(RenderTargetResolution);
+	}
+	else if (PropertyName == GET_MEMBER_NAME_CHECKED(ThisClass, ReflectionStrength))
+	{
+		SetReflectionStrength(ReflectionStrength);
+	}
+}
+
+void UCompositeLayerPlanarReflection::PostEditUndo()
+{
+	Super::PostEditUndo();
+
+	SetActors(MoveTemp(Actors));
+
+	SetCustomRenderPass(bCustomRenderPass);
+}
+
+void UCompositeLayerPlanarReflection::PostTransacted(const FTransactionObjectEvent& TransactionEvent)
+{
+	Super::PostTransacted(TransactionEvent);
+
+	if (TransactionEvent.GetEventType() == ETransactionObjectEventType::UndoRedo &&
+		TransactionEvent.GetChangedProperties().Contains(GET_MEMBER_NAME_CHECKED(ThisClass, RenderTargetResolution)))
+	{
+		/**
+		* When a resolution change is undone, the previous render targets are deserialized & replaced.
+		* This leaves previous assignments invalid so we make sure to return them to the pool.
+		*/
+		if (UCompositeSceneCapture2DComponent* CaptureComponent = GetSceneCaptureComponent())
+		{
+			FCompositeRenderTargetPool::Get().ReleaseAssigneeTargets(CaptureComponent);
+		}
+	}
+}
+#endif
+
+void UCompositeLayerPlanarReflection::OnEndOfFrameUpdate(UWorld* InWorld)
+{
+	/**
+	* In order to retain the ability to reference actors in sublevels, we prefer using weak pointer components on planar reflections. However,
+	* primitive components may often be recreated (see RerunConstructionScripts(), FStaticMeshComponentRecreateRenderStateContext, dynamic
+	* text primitives, etc). To prevent issues with destroyed primitives disappearing in our planar reflections, we re-register them every frame.
+	*/
+	USceneCaptureComponent2D* PlanarReflectionComponent = GetSceneCaptureComponent();
+	if (IsValid(PlanarReflectionComponent))
+	{
+		UpdateSceneCaptureComponents(*PlanarReflectionComponent, Actors);
+
+		UpdatePlanarReflectionTransform(PlanarReflectionComponent);
+	}
+}
+
+bool UCompositeLayerPlanarReflection::GetProxy(const UE::CompositeCore::FPassInputDecl& /*InputDecl*/, FCompositeTraversalContext& InContext, FSceneRenderingBulkObjectAllocator& InFrameAllocator, FCompositeCorePassProxy*& OutProxy) const
+{
+	using namespace UE::CompositeCore;
+
+	ACompositeActor* CompositeActor = GetTypedOuter<ACompositeActor>();
+	if (!IsValid(CompositeActor))
+	{
+		return false;
+	}
+
+	USceneCaptureComponent2D* CaptureComponent = GetSceneCaptureComponent();
+	if (!IsValid(CaptureComponent))
+	{
+		return false;
+	}
+
+	FCompositeRenderTargetPool::Get().ConditionalAcquireTarget(CaptureComponent, CaptureComponent->TextureTarget, RenderTargetResolution);
+
+	FResourceMetadata Metadata = {};
+	Metadata.bInvertedAlpha = true;
+	Metadata.bDistorted = CaptureComponent->ShowFlags.LensDistortion;
+	Metadata.DebugName = TEXT("CompositePlanarReflectionTex");
+
+	const ResourceId TexId = InContext.FindOrCreateExternalTexture(CaptureComponent->TextureTarget, Metadata);
+
+	FPassInputDecl PassInput;
+	PassInput.Set<FPassExternalResourceDesc>({ TexId });
+
+	FPlanarReflectionPassProxy* PlanarReflectionProxy = InFrameAllocator.Create<FPlanarReflectionPassProxy>(FPassInputDeclArray{ PassInput });
+	PlanarReflectionProxy->ReflectionStrength = ReflectionStrength;
+	PassInput.Set<const FCompositeCorePassProxy*>(PlanarReflectionProxy);
+
+	AddChildPasses(PassInput, InContext, InFrameAllocator, LayerPasses);
+
+	FPassInputDeclArray Inputs;
+	Inputs.SetNum(FixedNumLayerInputs);
+	Inputs[0] = PassInput;
+	Inputs[1] = GetDefaultSecondInput(InContext);
+
+	OutProxy = InFrameAllocator.Create<FMergePassProxy>(MoveTemp(Inputs), GetMergeOperation(InContext), TEXT("PlanarReflection"), ELensDistortionHandling::Disabled);
+	return true;
+}
+
+const TArray<TSoftObjectPtr<AActor>> UCompositeLayerPlanarReflection::GetActors() const
+{
+	return Actors;
+}
+
+void UCompositeLayerPlanarReflection::SetActors(TArray<TSoftObjectPtr<AActor>> InActors)
+{
+	Actors = MoveTemp(InActors);
+
+	USceneCaptureComponent2D* PlanarReflectionComponent = GetSceneCaptureComponent();
+	if (IsValid(PlanarReflectionComponent))
+	{
+		UpdateSceneCaptureComponents(*PlanarReflectionComponent, Actors);
+	}
+}
+
+bool UCompositeLayerPlanarReflection::IsCustomRenderPass() const
+{
+	return bCustomRenderPass;
+}
+
+void UCompositeLayerPlanarReflection::SetCustomRenderPass(bool bInIsFastRenderPass)
+{
+	bCustomRenderPass = bInIsFastRenderPass;
+
+	ACompositeActor* CompositeActor = GetTypedOuter<ACompositeActor>();
+	UCompositeSceneCapture2DComponent* PlanarReflectionComponent = GetSceneCaptureComponent();
+	if (IsValid(CompositeActor) && IsValid(PlanarReflectionComponent))
+	{
+		PlanarReflectionComponent->CaptureSource = UE::Composite::Private::GetPlanarReflectionCaptureSource(*CompositeActor, bCustomRenderPass);
+		PlanarReflectionComponent->bRenderInMainRenderer = bCustomRenderPass;
+	}
+}
+
+void UCompositeLayerPlanarReflection::SetRenderTargetResolution(FIntPoint InRenderTargetResolution)
+{
+	RenderTargetResolution.X = FMath::Max(InRenderTargetResolution.X, 1);
+	RenderTargetResolution.Y = FMath::Max(InRenderTargetResolution.Y, 1);
+}
+
+void UCompositeLayerPlanarReflection::SetReflectionStrength(float InReflectionStrength)
+{
+	ReflectionStrength = FMath::Clamp(InReflectionStrength, 0.0f, 1.0f);
+}
+
+UCompositeSceneCapture2DComponent* UCompositeLayerPlanarReflection::GetSceneCaptureComponent() const
+{
+	ACompositeActor* CompositeActor = GetTypedOuter<ACompositeActor>();
+	if (IsValid(CompositeActor))
+	{
+		UCompositeSceneCapture2DComponent* CaptureComponent = CompositeActor->FindOrCreateSceneCapture<UCompositeSceneCapture2DComponent>(this);
+		if (IsValid(CaptureComponent))
+		{
+			CaptureComponent->CaptureSource = UE::Composite::Private::GetPlanarReflectionCaptureSource(*CompositeActor, bCustomRenderPass);
+			CaptureComponent->bRenderInMainRenderer = bCustomRenderPass;
+			CaptureComponent->bMainViewCamera = false;
+			CaptureComponent->ShowFlags.SetAllowPrimitiveAlphaHoldout(true);
+		}
+
+		return CaptureComponent;
+	}
+
+	return nullptr;
+}
+
+void UCompositeLayerPlanarReflection::UpdatePlanarReflectionTransform(USceneCaptureComponent2D* CaptureComponent)
+{
+	if (!CaptureComponent)
+	{
+		return;
+	}
+
+	ACompositeActor* CompositeActor = GetTypedOuter<ACompositeActor>();
+	if (!CompositeActor)
+	{
+		return;
+	}
+
+	FComponentReference& CameraRef = CompositeActor->GetCamera();
+	UCameraComponent* CameraComponent = Cast<UCameraComponent>(CameraRef.GetComponent(nullptr));
+	const UWorld* World = GetWorld();
+
+	if (!World)
+	{
+		return;
+	}
+
+	FMinimalViewInfo ViewInfo;
+	bool bHasEditorViewportView = false;
+
+#if WITH_EDITOR
+	// Media profile viewports provide their own main view camera source, so we skip the editor viewport fallback in that mode.
+	if (CompositeActor->AllowedViewModes != ECompositeAllowedViewModes::MediaProfileUnknown)
+	{
+		bHasEditorViewportView = UE::Composite::Private::TryGetEditorViewportCameraView(*CompositeActor, ViewInfo);
+	}
+#endif
+
+	if (!bHasEditorViewportView)
+	{
+		if (!CameraComponent)
+		{
+			return;
+		}
+
+		CameraComponent->GetCameraView(World->DeltaTimeSeconds, ViewInfo);
+	}
+
+	// Calculate reflected position (mirror across horizontal plane at PlaneHeight)
+	float DistanceAboveGround = ViewInfo.Location.Z - PlaneHeight;
+	ViewInfo.Location.Z = PlaneHeight - DistanceAboveGround;
+
+	// Calculate reflected rotation (flip pitch and roll, preserve yaw)
+	ViewInfo.Rotation.Pitch = -ViewInfo.Rotation.Pitch;
+	ViewInfo.Rotation.Roll = -ViewInfo.Rotation.Roll;
+
+	CaptureComponent->SetCameraView(ViewInfo);
+
+	/**
+	* Preserve the source view projection so non-matching render target aspect ratios do not distort the reflection.
+	*
+	* Note: While we could obtain the correct y-flip by adjusting the projection matrix, doing so
+	* causes incorrect face culling order. It is therefore simpler to flip the image in our proxy shader.
+	*/
+	CaptureComponent->bUseCustomProjectionMatrix = true;
+	CaptureComponent->CustomProjectionMatrix = ViewInfo.CalculateProjectionMatrix();
+}
+
diff --git a/Engine/Plugins/Compositing/Composite/Source/Composite/Public/Layers/CompositeLayerPlanarReflection.h b/Engine/Plugins/Compositing/Composite/Source/Composite/Public/Layers/CompositeLayerPlanarReflection.h
new file mode 100644
--- /dev/null
+++ b/Engine/Plugins/Compositing/Composite/Source/Composite/Public/Layers/CompositeLayerPlanarReflection.h
@@ -0,0 +1,114 @@
+// Copyright Epic Games, Inc. All Rights Reserved.
+
+#pragma once
+
+#include "CompositeLayerBase.h"
+
+#include "CompositeLayerPlanarReflection.generated.h"
+
+#define UE_API COMPOSITE_API
+
+class AActor;
+class ACompositeActor;
+class USceneCaptureComponent2D;
+
+/**
+* Layer for planar reflection of selected actors.
+*/
+UCLASS(MinimalAPI, BlueprintType, Blueprintable, EditInlineNew, CollapseCategories, meta = (DisplayName = "Planar Reflection"))
+class UCompositeLayerPlanarReflection : public UCompositeLayerBase
+{
+	GENERATED_BODY()
+
+public:
+	/** Constructor. */
+	UE_API UCompositeLayerPlanarReflection(const FObjectInitializer& ObjectInitializer);
+
+	/** Destructor. */
+	UE_API virtual ~UCompositeLayerPlanarReflection();
+
+	UE_API virtual void OnRemoved(ACompositeActor* LastOwner) override;
+
+	//~ Begin UObject Interface
+	UE_API virtual void BeginDestroy() override;
+
+#if WITH_EDITOR
+	UE_API virtual void PostEditChangeProperty(FPropertyChangedEvent& PropertyChangedEvent) override;
+	UE_API virtual void PostEditUndo() override;
+	UE_API virtual void PostTransacted(const FTransactionObjectEvent& TransactionEvent) override;
+#endif
+	//~ End UObject Interface
+
+	/** End-of-frame update callback. */
+	UE_API virtual void OnEndOfFrameUpdate(UWorld* InWorld) override;
+
+	/** Getter function to override, returning pass layer proxies to be passed to the render thread. (Proxy objects should be allocated from the provided allocator.) */
+	UE_API virtual bool GetProxy(const UE::CompositeCore::FPassInputDecl& InputDecl, FCompositeTraversalContext& InContext, FSceneRenderingBulkObjectAllocator& InFrameAllocator, FCompositeCorePassProxy*& OutProxy) const override;
+
+	/** Get reflection & shadow caster actors. */
+	UFUNCTION(BlueprintGetter)
+	UE_API const TArray<TSoftObjectPtr<AActor>> GetActors() const;
+
+	/** Set reflection & shadow caster actors. */
+	UFUNCTION(BlueprintSetter)
+	UE_API void SetActors(TArray<TSoftObjectPtr<AActor>> InActors);
+
+	/** Get whether the planar reflection is rendered as a custom render pass. */
+	UFUNCTION(BlueprintGetter)
+	UE_API bool IsCustomRenderPass() const;
+
+	/** Set whether the planar reflection is rendered as a custom render pass. */
+	UFUNCTION(BlueprintSetter)
+	UE_API void SetCustomRenderPass(bool bInIsFastRenderPass);
+
+	/** Get the render target resolution for the planar reflection. */
+	UFUNCTION(BlueprintGetter)
+	UE_API FIntPoint GetRenderTargetResolution() const { return RenderTargetResolution; }
+
+	/** Set the render target resolution for the planar reflection. */
+	UFUNCTION(BlueprintSetter)
+	UE_API void SetRenderTargetResolution(FIntPoint InRenderTargetResolution);
+
+	/** Get the planar reflection strength. */
+	UFUNCTION(BlueprintGetter)
+	UE_API float GetReflectionStrength() const { return ReflectionStrength; }
+
+	/** Set the planar reflection strength. */
+	UFUNCTION(BlueprintSetter)
+	UE_API void SetReflectionStrength(float InReflectionStrength);
+
+	/** Sub-passes applied during post-processing. */
+	UPROPERTY(EditAnywhere, BlueprintReadWrite, Instanced, Category = "", meta = (DisallowedClasses = "/Script/Composite.CompositeLayerBase, /Script/Composite.CompositePassDistortion", DisplayAfter = "bCustomRenderPass"))
+	TArray<TObjectPtr<UCompositePassBase>> LayerPasses;
+
+private:
+
+	/** Convenience function to access the actor-managed planar reflection component. */
+	class UCompositeSceneCapture2DComponent* GetSceneCaptureComponent() const;
+
+	/** Update the scene capture transform for planar reflection. */
+	void UpdatePlanarReflectionTransform(USceneCaptureComponent2D* CaptureComponent);
+
+	/** List of actors to be rendered from the planar reflection. */
+	UPROPERTY(EditAnywhere, BlueprintGetter = GetActors, BlueprintSetter = SetActors, Category = "", meta = (AllowPrivateAccess = true))
+	TArray<TSoftObjectPtr<AActor>> Actors;
+
+	/** Render target resolution for the planar reflection. */
+	UPROPERTY(EditAnywhere, BlueprintGetter = GetRenderTargetResolution, BlueprintSetter = SetRenderTargetResolution, Interp, Category = "", meta = (AllowPrivateAccess = true))
+	FIntPoint RenderTargetResolution;
+
+	/** Strength of the planar reflection. */
+	UPROPERTY(EditAnywhere, BlueprintGetter = GetReflectionStrength, BlueprintSetter = SetReflectionStrength, Category = "", meta = (AllowPrivateAccess = true, ClampMin = "0.0", ClampMax = "1.0", UIMin = "0.0", UIMax = "1.0"))
+	float ReflectionStrength = 0.5f;
+
+	/** World-space Z coordinate of the horizontal plane for planar reflection. Note that we currently do not clip against this plane. */
+	UPROPERTY(EditAnywhere, BlueprintReadWrite, Category = "", meta = (AllowPrivateAccess = true))
+	float PlaneHeight = 0.0f;
+
+	/** Convenience toggle to render the planar reflection as a fast/minimal Custom Render Pass, inlined in the main render & without support for lighting. */
+	UPROPERTY(EditAnywhere, BlueprintGetter = IsCustomRenderPass, BlueprintSetter = SetCustomRenderPass, Category = "", AdvancedDisplay, meta = (AllowPrivateAccess = true))
+	bool bCustomRenderPass;
+};
+
+#undef UE_API
+
