Composure: Added custom widget for the pass list in the Scene Capture layer to match the pass tree widget in the Plate layer. This required some refactoring of the underlying pass tree widget to support multiple types of pass list owners instead of being specifically for plate layers

#Synced-CL 51311626

diff --git a/Engine/Plugins/Compositing/Composite/Source/CompositeEditor/Private/Customizations/CompositeLayerPlateCustomization.cpp b/Engine/Plugins/Compositing/Composite/Source/CompositeEditor/Private/Customizations/CompositeLayerPlateCustomization.cpp
--- a/Engine/Plugins/Compositing/Composite/Source/CompositeEditor/Private/Customizations/CompositeLayerPlateCustomization.cpp
+++ b/Engine/Plugins/Compositing/Composite/Source/CompositeEditor/Private/Customizations/CompositeLayerPlateCustomization.cpp
@@ -27,8 +27,9 @@
 #include "Profile/MediaProfilePlaybackManager.h"
 #include "Styling/SlateIconFinder.h"
 #include "Subsystems/AssetEditorSubsystem.h"
+#include "UI/CompositePlatePassListOwner.h"
 #include "UI/SCompositeActorPickerTable.h"
-#include "UI/SCompositePlatePassPanel.h"
+#include "UI/SCompositePassTreePanel.h"
 #include "Widgets/Images/SImage.h"
 #include "Widgets/Input/SComboButton.h"
 #include "Widgets/Text/STextBlock.h"
@@ -55,9 +56,17 @@ void FCompositeLayerPlateCustomization::CustomizeDetails(IDetailLayoutBuilder& D
 {
 	TArray<TWeakObjectPtr<UCompositeLayerPlate>> Objects = DetailLayout.GetObjectsOfTypeBeingCustomized<UCompositeLayerPlate>();
 
-	DetailLayout.HideProperty(GET_MEMBER_NAME_CHECKED(UCompositeLayerPlate, CompositeMeshes));
-	DetailLayout.HideProperty(GET_MEMBER_NAME_CHECKED(UCompositeLayerPlate, MediaPasses));
-	DetailLayout.HideProperty(GET_MEMBER_NAME_CHECKED(UCompositeLayerPlate, LayerPasses));
+	static TArray<FName> HiddenProperties =
+	{
+		GET_MEMBER_NAME_CHECKED(UCompositeLayerPlate, CompositeMeshes),
+		GET_MEMBER_NAME_CHECKED(UCompositeLayerPlate, MediaPasses),
+		GET_MEMBER_NAME_CHECKED(UCompositeLayerPlate, LayerPasses)
+	};
+
+	for (const FName& PropertyName : HiddenProperties)
+	{
+		DetailLayout.HideProperty(PropertyName);
+	}
 	
 	IDetailCategoryBuilder& PlateCategory = DetailLayout.EditCategoryAllowNone(NAME_None);
 
@@ -74,7 +83,7 @@ void FCompositeLayerPlateCustomization::CustomizeDetails(IDetailLayoutBuilder& D
 	PlateCategory.GetDefaultProperties(PlateProperties, /* bSimpleProperties */ true, /* bAdvancedProperties */ false);
 	for (const TSharedRef<IPropertyHandle>& Property : PlateProperties)
 	{
-		if (!Property->IsCustomized())
+		if (!HiddenProperties.Contains(Property->GetProperty()->GetNameCPP()))
 		{
 			const FName PropertyName = Property->GetProperty()->GetFName();
 			IDetailPropertyRow& PropertyRow = PlateCategory.AddProperty(Property);
@@ -107,7 +116,7 @@ void FCompositeLayerPlateCustomization::CustomizeDetails(IDetailLayoutBuilder& D
 		PassesGroup.AddWidgetRow()
 		.WholeRowContent()
 		[
-			SAssignNew(PassPanel, SCompositePlatePassPanel, Plate)
+			SAssignNew(PassPanel, SCompositePassTreePanel, MakeShared<FCompositePlatePassListOwner>(Plate))
 			.OnLayoutSizeChanged(this, &FCompositeLayerPlateCustomization::OnLayoutSizeChanged)
 		];
 	}  
@@ -143,7 +152,7 @@ bool FCompositeLayerPlateCustomization::IsCustomizingObject(UObject* InObject) c
 	return false;
 }
 
-TSharedPtr<SCompositePlatePassPanel> FCompositeLayerPlateCustomization::GetPassPanelWidget() const
+TSharedPtr<SCompositePassTreePanel> FCompositeLayerPlateCustomization::GetPassPanelWidget() const
 {
 	return PassPanel.Pin();
 }
diff --git a/Engine/Plugins/Compositing/Composite/Source/CompositeEditor/Private/Customizations/CompositeLayerPlateCustomization.h b/Engine/Plugins/Compositing/Composite/Source/CompositeEditor/Private/Customizations/CompositeLayerPlateCustomization.h
--- a/Engine/Plugins/Compositing/Composite/Source/CompositeEditor/Private/Customizations/CompositeLayerPlateCustomization.h
+++ b/Engine/Plugins/Compositing/Composite/Source/CompositeEditor/Private/Customizations/CompositeLayerPlateCustomization.h
@@ -8,7 +8,7 @@ class SWidget;
 class IDetailPropertyRow;
 class IPropertyHandle;
 class SCompositeActorPickerTable;
-class SCompositePlatePassPanel;
+class SCompositePassTreePanel;
 class FMenuBuilder;
 
 /**
@@ -26,7 +26,7 @@ public:
 	// End of IDetailCustomization interface
 
 	bool IsCustomizingObject(UObject* InObject) const;
-	TSharedPtr<SCompositePlatePassPanel> GetPassPanelWidget() const;
+	TSharedPtr<SCompositePassTreePanel> GetPassPanelWidget() const;
 	TSharedPtr<SCompositeActorPickerTable> GetCompositeMeshPickerWidget() const;
 	
 private:
@@ -59,7 +59,7 @@ private:
 	
 private:
 	TWeakPtr<IDetailLayoutBuilder> CachedDetailBuilder;
-	TWeakPtr<SCompositePlatePassPanel> PassPanel;
+	TWeakPtr<SCompositePassTreePanel> PassPanel;
 	TWeakPtr<SCompositeActorPickerTable> CompositeMeshPicker;
 	TSharedPtr<IPropertyHandle> TexturePropertyHandle;
 };
diff --git a/Engine/Plugins/Compositing/Composite/Source/CompositeEditor/Private/Customizations/CompositeLayerSceneCaptureCustomization.cpp b/Engine/Plugins/Compositing/Composite/Source/CompositeEditor/Private/Customizations/CompositeLayerSceneCaptureCustomization.cpp
--- a/Engine/Plugins/Compositing/Composite/Source/CompositeEditor/Private/Customizations/CompositeLayerSceneCaptureCustomization.cpp
+++ b/Engine/Plugins/Compositing/Composite/Source/CompositeEditor/Private/Customizations/CompositeLayerSceneCaptureCustomization.cpp
@@ -9,7 +9,9 @@
 #include "IDetailGroup.h"
 #include "IPropertyUtilities.h"
 #include "Layers/CompositeLayerSceneCapture.h"
+#include "UI/CompositeSceneCapturePassListOwner.h"
 #include "UI/SCompositeActorPickerTable.h"
+#include "UI/SCompositePassTreePanel.h"
 #include "Widgets/Text/STextBlock.h"
 
 #define LOCTEXT_NAMESPACE "FCompositeLayerSceneCaptureCustomization"
@@ -33,7 +35,16 @@ void FCompositeLayerSceneCaptureCustomization::CustomizeDetails(IDetailLayoutBui
 		}),
 		MakeShared<FResolutionPropertyIdentifier>());
 
-	DetailLayout.HideProperty(GET_MEMBER_NAME_CHECKED(UCompositeLayerSceneCapture, Actors));
+	static TArray<FName> HiddenProperties =
+	{
+		GET_MEMBER_NAME_CHECKED(UCompositeLayerSceneCapture, Actors),
+		GET_MEMBER_NAME_CHECKED(UCompositeLayerSceneCapture, LayerPasses)
+	};
+
+	for (const FName& PropertyName : HiddenProperties)
+	{
+		DetailLayout.HideProperty(PropertyName);
+	}
 	
 	TArray<TWeakObjectPtr<UCompositeLayerSceneCapture>> Objects = DetailLayout.GetObjectsOfTypeBeingCustomized<UCompositeLayerSceneCapture>();
 
@@ -44,7 +55,7 @@ void FCompositeLayerSceneCaptureCustomization::CustomizeDetails(IDetailLayoutBui
 	LayerCategory.GetDefaultProperties(LayerProperties, /* bSimpleProperties */ true, /* bAdvancedProperties */ false);
 	for (const TSharedRef<IPropertyHandle>& Property : LayerProperties)
 	{
-		if (!Property->IsCustomized())
+		if (!HiddenProperties.Contains(Property->GetProperty()->GetNameCPP()))
 		{
 			LayerCategory.AddProperty(Property);
 		}
@@ -54,6 +65,15 @@ void FCompositeLayerSceneCaptureCustomization::CustomizeDetails(IDetailLayoutBui
 	{
 		UCompositeLayerSceneCapture* SceneCapture = Cast<UCompositeLayerSceneCapture>(Objects[0].Get());
 		
+		IDetailGroup& PassesGroup = LayerCategory.AddGroup("Passes", LOCTEXT("PassesGroupName", "Passes"), false, true);
+
+		PassesGroup.AddWidgetRow()
+		.WholeRowContent()
+		[
+			SAssignNew(PassPanel, SCompositePassTreePanel, MakeShared<FCompositeSceneCapturePassListOwner>(SceneCapture))
+			.OnLayoutSizeChanged(this, &FCompositeLayerSceneCaptureCustomization::OnLayoutSizeChanged)
+		];
+		
 		IDetailGroup& ActorListGroup = LayerCategory.AddGroup("ActorContent", LOCTEXT("ActorListGroupName", "Scene Capture Content"), false, true);
 		FCompositeActorPickerListRef ActorListRef(SceneCapture, GET_MEMBER_NAME_CHECKED(UCompositeLayerSceneCapture, Actors), &SceneCapture->Actors);
 		
@@ -66,6 +86,21 @@ void FCompositeLayerSceneCaptureCustomization::CustomizeDetails(IDetailLayoutBui
 	}
 	else
 	{
+		// Can't display pass panel if multiple layers are selected, so simply put a "Multiple Values" entry in the property list
+		LayerCategory.AddCustomRow(LOCTEXT("PassesGroupName", "Passes"))
+		.NameContent()
+		[
+			SNew(STextBlock)
+			.Text(LOCTEXT("PassesGroupName", "Passes"))
+			.Font(IDetailLayoutBuilder::GetDetailFont())
+		]
+		.ValueContent()
+		[
+			SNew(STextBlock)
+			.Text(LOCTEXT("MultipleValues", "Multiple Values"))
+			.Font(IDetailLayoutBuilder::GetDetailFont())
+		];
+		
 		// Can't display actor list if multiple layers are selected, so simply put a "Multiple Values" entry in the property list
 		LayerCategory.AddCustomRow(LOCTEXT("ActorListGroupName", "Scene Capture Content"))
 		.NameContent()
@@ -83,6 +118,24 @@ void FCompositeLayerSceneCaptureCustomization::CustomizeDetails(IDetailLayoutBui
 	}
 }
 
+bool FCompositeLayerSceneCaptureCustomization::IsCustomizingObject(UObject* InObject) const
+{
+	if (TSharedPtr<IDetailLayoutBuilder> Builder = CachedDetailBuilder.Pin())
+	{
+		TArray<TWeakObjectPtr<UObject>> ObjectsBeingCustomized;
+		Builder->GetObjectsBeingCustomized(ObjectsBeingCustomized);
+
+		return ObjectsBeingCustomized.Contains(InObject);
+	}
+
+	return false;
+}
+
+TSharedPtr<SCompositePassTreePanel> FCompositeLayerSceneCaptureCustomization::GetPassPanelWidget() const
+{
+	return PassPanel.Pin();
+}
+
 void FCompositeLayerSceneCaptureCustomization::OnLayoutSizeChanged()
 {
 	TSharedPtr<IDetailLayoutBuilder> PinnedDetailBuilder = CachedDetailBuilder.Pin();
diff --git a/Engine/Plugins/Compositing/Composite/Source/CompositeEditor/Private/Customizations/CompositeLayerSceneCaptureCustomization.h b/Engine/Plugins/Compositing/Composite/Source/CompositeEditor/Private/Customizations/CompositeLayerSceneCaptureCustomization.h
--- a/Engine/Plugins/Compositing/Composite/Source/CompositeEditor/Private/Customizations/CompositeLayerSceneCaptureCustomization.h
+++ b/Engine/Plugins/Compositing/Composite/Source/CompositeEditor/Private/Customizations/CompositeLayerSceneCaptureCustomization.h
@@ -4,6 +4,7 @@
 
 #include "IDetailCustomization.h"
 
+class SCompositePassTreePanel;
 /**
  * Customization for the composite scene capture layer, primary for displaying a custom widget for picking actors
  */
@@ -18,10 +19,17 @@ public:
 	virtual void CustomizeDetails(IDetailLayoutBuilder& DetailLayout) override;
 	// End of IDetailCustomization interface
 
+	/** Gets whether the specified object is one of the objects being customized by this details panel customization */
+	bool IsCustomizingObject(UObject* InObject) const;
+
+	/** Gets the pass panel widget the details panel is displaying */
+	TSharedPtr<SCompositePassTreePanel> GetPassPanelWidget() const;
+	
 private:
 	/** Raised when the custom widgets' layout size may have changed */
 	void OnLayoutSizeChanged();
 	
 private:
 	TWeakPtr<IDetailLayoutBuilder> CachedDetailBuilder;
+	TWeakPtr<SCompositePassTreePanel> PassPanel;
 };
diff --git a/Engine/Plugins/Compositing/Composite/Source/CompositeEditor/Private/UI/CompositePlatePassListOwner.cpp b/Engine/Plugins/Compositing/Composite/Source/CompositeEditor/Private/UI/CompositePlatePassListOwner.cpp
new file mode 100644
--- /dev/null
+++ b/Engine/Plugins/Compositing/Composite/Source/CompositeEditor/Private/UI/CompositePlatePassListOwner.cpp
@@ -0,0 +1,314 @@
+﻿// Copyright Epic Games, Inc. All Rights Reserved.
+
+#include "CompositePlatePassListOwner.h"
+
+#include "CompositeEditorStyle.h"
+#include "PropertyEditorUtils.h"
+#include "Layers/CompositeLayerPlate.h"
+#include "Passes/CompositePassMaterial.h"
+#include "UObject/Package.h"
+
+#define LOCTEXT_NAMESPACE "FCompositePlatePassListOwner"
+
+bool FCompositePlatePassListOwner::IsObjectValid()
+{
+	return Plate.IsValid();
+}
+
+TStrongObjectPtr<UObject> FCompositePlatePassListOwner::GetObject()
+{
+	return Plate.Pin();
+}
+
+bool FCompositePlatePassListOwner::IsPassListPropertyName(const FName& InPropertyName)
+{
+	return
+		InPropertyName == GET_MEMBER_NAME_CHECKED(UCompositeLayerPlate, MediaPasses) ||
+		InPropertyName == GET_MEMBER_NAME_CHECKED(UCompositeLayerPlate, LayerPasses);
+}
+
+bool FCompositePlatePassListOwner::IsValidPassIndex(int32 InGroupIndex, int32 InPassIndex)
+{
+	if (!Plate.IsValid())
+	{
+		return false;
+	}
+	
+	TArray<TObjectPtr<UCompositePassBase>>& PassList = GetPassList((EPassType)InGroupIndex);
+	return PassList.IsValidIndex(InPassIndex);
+}
+
+UCompositePassBase* FCompositePlatePassListOwner::GetPass(int32 InGroupIndex, int32 InPassIndex)
+{
+	if (!Plate.IsValid())
+	{
+		return nullptr;
+	}
+	
+	TArray<TObjectPtr<UCompositePassBase>>& PassList = GetPassList((EPassType)InGroupIndex);
+
+	if (!PassList.IsValidIndex(InPassIndex))
+	{
+		return nullptr;
+	}
+	
+	return PassList[InPassIndex];
+}
+
+TArray<TObjectPtr<UCompositePassBase>>& FCompositePlatePassListOwner::GetPassesForGroup(int32 InGroupIndex)
+{
+	return GetPassList((EPassType)InGroupIndex);
+}
+
+FString FCompositePlatePassListOwner::GetGroupFilterString(int32 InGroupIndex)
+{
+	static FString PassTypeFilterStrings[] =
+	{
+		TEXT("Media Passes"),
+		TEXT("Layer Passes")
+	};
+
+	return PassTypeFilterStrings[InGroupIndex];
+}
+
+const FSlateBrush* FCompositePlatePassListOwner::GetGroupIcon(int32 InGroupIndex)
+{
+	static const FSlateBrush* PassTypeIcons[] =
+	{
+		FCompositeEditorStyle::Get().GetBrush("CompositeEditor.Passes.Media"),
+		FCompositeEditorStyle::Get().GetBrush("CompositeEditor.Passes.Layer")
+	};
+		
+	return PassTypeIcons[InGroupIndex];
+}
+
+FText FCompositePlatePassListOwner::GetGroupDisplayName(int32 InGroupIndex)
+{
+	static FText PassTypeLabels[] =
+	{
+		LOCTEXT("MediaPassTypeLabel", "Media Passes"),
+		LOCTEXT("LayerPassTypeLabel", "Layer Passes"),
+	};
+		
+	return PassTypeLabels[InGroupIndex];
+}
+
+ICompositePassListOwner::FGroupFilterConfig FCompositePlatePassListOwner::GetGroupFilterConfig(int32 InGroupIndex)
+{
+	static TArray<FGroupFilterConfig> PassTypeFilterConfigs =
+	{
+		{ TEXT("MediaPassesFilter"), LOCTEXT("MediaPassesFilterName", "Media Passes"),  LOCTEXT("MediaPassesFilterToolTip", "Only show media passes") },
+		{ TEXT("LayerPassesFilter"), LOCTEXT("LayerPassesFilterName", "Layer Passes"),  LOCTEXT("LayerPassesFilterToolTip", "Only show layer passes")  }
+	};
+
+	return PassTypeFilterConfigs[InGroupIndex];
+}
+
+int32 FCompositePlatePassListOwner::GetDefaultGroupForNewPass(const UClass* InPassClass) const
+{
+	// The default group to add the pass to is the media passes list, unless it is a PPM pass, which can only get added to the layers passes list
+	if (InPassClass->IsChildOf<UCompositePassMaterial>())
+	{
+		return (int32)EPassType::Layer;
+	}
+	else
+	{
+		return (int32)EPassType::Media;
+	}
+}
+
+bool FCompositePlatePassListOwner::CanAddPass(const UClass* InPassClass, int32 InGroupIndex, int32 InPassIndex) const
+{
+	TStrongObjectPtr<UCompositeLayerPlate> PinnedPlate = Plate.Pin();
+	if (!PinnedPlate.IsValid())
+	{
+		return false;
+	}
+
+	const EPassType DestPassType = (EPassType)InGroupIndex;
+	FProperty* PassListProperty = GetPassListProperty(DestPassType);
+
+	TArray<const UClass*> AllowedClassFilters;
+	TArray<const UClass*> DisallowedClassFilters;
+	PropertyEditorUtils::GetAllowedAndDisallowedClasses({ PinnedPlate.Get() }, *PassListProperty, AllowedClassFilters, DisallowedClassFilters, false);
+
+	bool bIsAllowed = AllowedClassFilters.IsEmpty() || AllowedClassFilters.Contains(InPassClass);
+	bool bIsDisallowed = DisallowedClassFilters.Contains(InPassClass);
+	
+	return bIsAllowed && !bIsDisallowed;
+}
+
+int32 FCompositePlatePassListOwner::AddPass(const UClass* InPassClass, int32 InGroupIndex, int32 InPassIndex)
+{
+	TStrongObjectPtr<UCompositeLayerPlate> PinnedPlate = Plate.Pin();
+	if (!PinnedPlate.IsValid())
+	{
+		return INDEX_NONE;
+	}
+	
+	PinnedPlate->Modify();
+
+	const EPassType DestPassType = (EPassType)InGroupIndex;
+	
+	UCompositePassBase* NewPass = NewObject<UCompositePassBase>(PinnedPlate.Get(), InPassClass, NAME_None, RF_Transactional);
+	TArray<TObjectPtr<UCompositePassBase>>& PassList = GetPassList(DestPassType);
+
+	FProperty* PassListProperty = GetPassListProperty(DestPassType);
+	PinnedPlate->PreEditChange(PassListProperty);
+	
+	int32 NewPassIndex = PassList.Insert(NewPass, InPassIndex);
+
+	FPropertyChangedEvent PropertyChangedEvent(PassListProperty, EPropertyChangeType::ArrayAdd);
+	PinnedPlate->PostEditChangeProperty(PropertyChangedEvent);
+
+	return NewPassIndex;
+}
+
+TArray<int32> FCompositePlatePassListOwner::CopyPasses(const TArray<UCompositePassBase*>& InPassesToCopy, int32 InGroupIndex, int32 InPassIndex)
+{
+	TArray<int32> OutNewPassIndices;
+	TStrongObjectPtr<UCompositeLayerPlate> PinnedPlate = Plate.Pin();
+	if (!PinnedPlate.IsValid())
+	{
+		return OutNewPassIndices;
+	}
+	
+	const EPassType PassTypeToCopyTo = (EPassType)InGroupIndex;
+	TArray<TObjectPtr<UCompositePassBase>>& PassList = GetPassList(PassTypeToCopyTo);
+	
+	FProperty* PassListProperty = GetPassListProperty(PassTypeToCopyTo);
+	PinnedPlate->PreEditChange(PassListProperty);
+
+	int32 IndexToCopyTo = InPassIndex;
+	for (UCompositePassBase* PassToPaste : InPassesToCopy)
+	{
+		UCompositePassBase* NewPass = DuplicateObject(PassToPaste, PinnedPlate.Get());
+			
+		const int32 NewPassIndex = PassList.Insert(NewPass, IndexToCopyTo);
+		OutNewPassIndices.Add(NewPassIndex);
+		++IndexToCopyTo;
+	}
+
+	FPropertyChangedEvent PropertyChangedEvent(PassListProperty, EPropertyChangeType::ArrayAdd);
+	PinnedPlate->PostEditChangeProperty(PropertyChangedEvent);
+
+	return OutNewPassIndices;
+}
+
+void FCompositePlatePassListOwner::MovePass(int32 InSourceGroupIndex, int32 InSourcePassIndex, int32 InDestGroupIndex, int32 InDestPassIndex)
+{
+	TStrongObjectPtr<UCompositeLayerPlate> PinnedPlate = Plate.Pin();
+	if (!PinnedPlate.IsValid())
+	{
+		return;
+	}
+
+	PinnedPlate->Modify();
+
+	const EPassType SourcePassType = (EPassType)InSourceGroupIndex;
+	const EPassType DestPassType = (EPassType)InDestGroupIndex;
+	
+	TArray<TObjectPtr<UCompositePassBase>>& SourcePassList = GetPassList(SourcePassType);
+	TArray<TObjectPtr<UCompositePassBase>>& DestPassList = GetPassList(DestPassType);
+
+	FProperty* SourcePassListProperty = GetPassListProperty(SourcePassType);
+	FProperty* DestPassListProperty = GetPassListProperty(DestPassType);
+	
+	UCompositePassBase* Pass = SourcePassList[InSourcePassIndex];
+	
+	PinnedPlate->PreEditChange(SourcePassListProperty);
+	SourcePassList.RemoveAt(InSourcePassIndex);
+
+	// If we aren't adding the item back to the same pass list, invoke a property changed event for its removal
+	if (SourcePassType != DestPassType)
+	{
+		FPropertyChangedEvent SourceListChangedEvent(SourcePassListProperty, EPropertyChangeType::ArrayRemove);
+		PinnedPlate->PostEditChangeProperty(SourceListChangedEvent);
+		PinnedPlate->PreEditChange(DestPassListProperty);
+	}
+	
+	DestPassList.Insert(Pass, InDestPassIndex);
+
+	// If the item moved from one list to another, invoke a property changed event for the add; otherwise, invoke a move event
+	if (SourcePassType != DestPassType)
+	{
+		FPropertyChangedEvent DestListChangedEvent(DestPassListProperty, EPropertyChangeType::ArrayAdd);
+		PinnedPlate->PostEditChangeProperty(DestListChangedEvent);
+	}
+	else
+	{
+		FPropertyChangedEvent MoveChangedEvent(SourcePassListProperty, EPropertyChangeType::ArrayMove);
+		PinnedPlate->PostEditChangeProperty(MoveChangedEvent);
+	}
+}
+
+void FCompositePlatePassListOwner::RemovePasses(int32 InGroupIndex, const TArray<int32>& InPassIndices)
+{
+	TStrongObjectPtr<UCompositeLayerPlate> PinnedPlate = Plate.Pin();
+	if (!PinnedPlate.IsValid())
+	{
+		return;
+	}
+
+	// Delete the passes in reverse order so that indices don't get messed up as we delete
+	TArray<int32> SortedPassIndices = InPassIndices;
+	SortedPassIndices.Sort();
+	
+	const EPassType PassType = (EPassType)InGroupIndex;
+	TArray<TObjectPtr<UCompositePassBase>>& PassList = GetPassList(PassType);
+
+	FProperty* PassListProperty = GetPassListProperty(PassType);
+	PinnedPlate->PreEditChange(PassListProperty);
+
+	for (int32 Index = SortedPassIndices.Num() - 1; Index >= 0; --Index)
+	{
+		const int32 PassIndex = SortedPassIndices[Index];
+
+		// Here we replicate details panels instanced property array removal code.
+		// The removed instanced pass objects must be moved to the transient package.
+		TObjectPtr<UCompositePassBase>& PassToRemove = PassList[PassIndex];
+		if (IsValid(PassToRemove))
+		{
+			PassToRemove->Modify();
+			PassToRemove->Rename(nullptr, GetTransientPackage(), REN_DontCreateRedirectors);
+		}
+
+		PassList.RemoveAt(PassIndex, EAllowShrinking::No);
+	}
+
+	PassList.Shrink();
+			
+	FPropertyChangedEvent PropertyChangedEvent(PassListProperty, EPropertyChangeType::ArrayAdd);
+	PinnedPlate->PostEditChangeProperty(PropertyChangedEvent);
+}
+
+TArray<TObjectPtr<UCompositePassBase>>& FCompositePlatePassListOwner::GetPassList(EPassType InPassType)
+{
+	TStrongObjectPtr<UCompositeLayerPlate> PinnedPlate = Plate.Pin();
+	check(PinnedPlate.IsValid());
+
+	switch (InPassType)
+	{
+	case EPassType::Media:
+		return PinnedPlate->MediaPasses;
+
+	case EPassType::Layer:
+	default:
+		return PinnedPlate->LayerPasses;
+	}
+}
+
+FProperty* FCompositePlatePassListOwner::GetPassListProperty(EPassType InPassType)
+{
+	switch (InPassType)
+	{
+	case EPassType::Media:
+		return FindFProperty<FProperty>(UCompositeLayerPlate::StaticClass(), GET_MEMBER_NAME_CHECKED(UCompositeLayerPlate, MediaPasses));
+
+	case EPassType::Layer:
+	default:
+		return FindFProperty<FProperty>(UCompositeLayerPlate::StaticClass(), GET_MEMBER_NAME_CHECKED(UCompositeLayerPlate, LayerPasses));
+	}
+}
+
+#undef LOCTEXT_NAMESPACE
diff --git a/Engine/Plugins/Compositing/Composite/Source/CompositeEditor/Private/UI/CompositePlatePassListOwner.h b/Engine/Plugins/Compositing/Composite/Source/CompositeEditor/Private/UI/CompositePlatePassListOwner.h
new file mode 100644
--- /dev/null
+++ b/Engine/Plugins/Compositing/Composite/Source/CompositeEditor/Private/UI/CompositePlatePassListOwner.h
@@ -0,0 +1,62 @@
+﻿// Copyright Epic Games, Inc. All Rights Reserved.
+
+#pragma once
+
+#include "CoreMinimal.h"
+#include "SCompositePassTree.h"
+
+/** ICompositePassListOwner implementation for UCompositeLayerPlate that manages its pass lists */
+class FCompositePlatePassListOwner : public ICompositePassListOwner
+{
+private:
+	/** Pass type which corresponds to which list in the plate layer a pass is stored in */
+	enum class EPassType : uint8
+	{
+		Media,
+		//Scene, // TODO: Add Scene when it is supported
+		Layer,
+		MAX
+	};
+
+public:
+	FCompositePlatePassListOwner(const TWeakObjectPtr<UCompositeLayerPlate>& InPlate)
+		: Plate(InPlate)
+	{ }
+
+	//~ ICompositePassListOwner interface
+	
+public:
+	virtual bool IsObjectValid() override;
+	virtual TStrongObjectPtr<UObject> GetObject() override;
+	virtual bool IsPassListPropertyName(const FName& InPropertyName) override;
+
+	virtual  bool IsValidPassIndex(int32 InGroupIndex, int32 InPassIndex) override;
+	virtual UCompositePassBase* GetPass(int32 InGroupIndex, int32 InPassIndex) override;
+	virtual int32 GetNumGroups() const override { return (int32)EPassType::MAX; }
+	virtual  TArray<TObjectPtr<UCompositePassBase>>& GetPassesForGroup(int32 InGroupIndex) override;
+	
+	virtual FString GetGroupFilterString(int32 InGroupIndex) override;
+	virtual const FSlateBrush* GetGroupIcon(int32 InGroupIndex) override;
+	virtual FText GetGroupDisplayName(int32 InGroupIndex) override;
+	virtual FGroupFilterConfig GetGroupFilterConfig(int32 InGroupIndex) override;
+
+	virtual int32 GetDefaultGroupForNewPass(const UClass* InPassClass) const override;
+	virtual bool CanAddPass(const UClass* InPassClass, int32 InGroupIndex, int32 InPassIndex) const override;
+	virtual int32 AddPass(const UClass* InPassClass, int32 InGroupIndex, int32 InPassIndex) override;
+	virtual TArray<int32> CopyPasses(const TArray<UCompositePassBase*>& InPassesToCopy, int32 InGroupIndex, int32 InPassIndex) override;
+	virtual void MovePass(int32 InSourceGroupIndex, int32 InSourcePassIndex, int32 InDestGroupIndex, int32 InDestPassIndex) override;
+	virtual void RemovePasses(int32 InGroupIndex, const TArray<int32>& InPassIndices) override;
+	
+	//~ End ICompositePassListOwner interface
+	
+private:
+	/** Gets a reference to the pass list that corresponds to the specified pass type in the plate. The plate pointer must be valid to call properly */
+	TArray<TObjectPtr<UCompositePassBase>>& GetPassList(EPassType InPassType);
+
+	/** Gets the FProperty object for the list that corresponds to the specified pass type in the plate */
+	static FProperty* GetPassListProperty(EPassType InPassType);
+	
+private:
+	/** The plate UObject whose pass lists are being provided */
+	TWeakObjectPtr<UCompositeLayerPlate> Plate;
+};
diff --git a/Engine/Plugins/Compositing/Composite/Source/CompositeEditor/Private/UI/CompositeSceneCapturePassListOwner.cpp b/Engine/Plugins/Compositing/Composite/Source/CompositeEditor/Private/UI/CompositeSceneCapturePassListOwner.cpp
new file mode 100644
--- /dev/null
+++ b/Engine/Plugins/Compositing/Composite/Source/CompositeEditor/Private/UI/CompositeSceneCapturePassListOwner.cpp
@@ -0,0 +1,221 @@
+﻿// Copyright Epic Games, Inc. All Rights Reserved.
+
+#include "CompositeSceneCapturePassListOwner.h"
+
+#include "PropertyEditorUtils.h"
+#include "Layers/CompositeLayerSceneCapture.h"
+#include "UObject/Package.h"
+
+bool FCompositeSceneCapturePassListOwner::IsObjectValid()
+{
+	return SceneCapture.IsValid();
+}
+
+TStrongObjectPtr<UObject> FCompositeSceneCapturePassListOwner::GetObject()
+{
+	return SceneCapture.Pin();
+}
+
+bool FCompositeSceneCapturePassListOwner::IsPassListPropertyName(const FName& InPropertyName)
+{
+	return InPropertyName == GET_MEMBER_NAME_CHECKED(UCompositeLayerSceneCapture, LayerPasses);
+}
+
+bool FCompositeSceneCapturePassListOwner::IsValidPassIndex(int32 InGroupIndex, int32 InPassIndex)
+{
+	if (InGroupIndex != INDEX_NONE)
+	{
+		return false;
+	}
+	
+	TStrongObjectPtr<UCompositeLayerSceneCapture> PinnedSceneCapture = SceneCapture.Pin();
+	if (!PinnedSceneCapture.IsValid())
+	{
+		return false;
+	}
+
+	return PinnedSceneCapture->LayerPasses.IsValidIndex(InPassIndex);
+}
+
+UCompositePassBase* FCompositeSceneCapturePassListOwner::GetPass(int32 InGroupIndex, int32 InPassIndex)
+{
+	if (InGroupIndex != INDEX_NONE)
+	{
+		return nullptr;
+	}
+	
+	TStrongObjectPtr<UCompositeLayerSceneCapture> PinnedSceneCapture = SceneCapture.Pin();
+	if (!PinnedSceneCapture.IsValid())
+	{
+		return nullptr;
+	}
+
+	if (!PinnedSceneCapture->LayerPasses.IsValidIndex(InPassIndex))
+	{
+		return nullptr;
+	}
+
+	return PinnedSceneCapture->LayerPasses[InPassIndex];
+}
+
+TArray<TObjectPtr<UCompositePassBase>>& FCompositeSceneCapturePassListOwner::GetPassesForGroup(int32 InGroupIndex)
+{
+	TStrongObjectPtr<UCompositeLayerSceneCapture> PinnedSceneCapture = SceneCapture.Pin();
+	check(PinnedSceneCapture.IsValid());
+
+	return PinnedSceneCapture->LayerPasses;
+}
+
+bool FCompositeSceneCapturePassListOwner::CanAddPass(const UClass* InPassClass, int32 InGroupIndex, int32 InPassIndex) const
+{
+	if (InGroupIndex != INDEX_NONE)
+	{
+		return false;
+	}
+
+	TStrongObjectPtr<UCompositeLayerSceneCapture> PinnedSceneCapture = SceneCapture.Pin();
+	if (!PinnedSceneCapture.IsValid())
+	{
+		return false;
+	}
+
+	FProperty* PassListProperty = FindFProperty<FProperty>(UCompositeLayerSceneCapture::StaticClass(), GET_MEMBER_NAME_CHECKED(UCompositeLayerSceneCapture, LayerPasses));
+
+	TArray<const UClass*> AllowedClassFilters;
+	TArray<const UClass*> DisallowedClassFilters;
+	PropertyEditorUtils::GetAllowedAndDisallowedClasses({ PinnedSceneCapture.Get() }, *PassListProperty, AllowedClassFilters, DisallowedClassFilters, false);
+
+	bool bIsAllowed = AllowedClassFilters.IsEmpty() || AllowedClassFilters.Contains(InPassClass);
+	bool bIsDisallowed = DisallowedClassFilters.Contains(InPassClass);
+	
+	return bIsAllowed && !bIsDisallowed;
+}
+
+int32 FCompositeSceneCapturePassListOwner::AddPass(const UClass* InPassClass, int32 InGroupIndex, int32 InPassIndex)
+{
+	if (InGroupIndex != INDEX_NONE)
+	{
+		return INDEX_NONE;
+	}
+
+	TStrongObjectPtr<UCompositeLayerSceneCapture> PinnedSceneCapture = SceneCapture.Pin();
+	if (!PinnedSceneCapture.IsValid())
+	{
+		return INDEX_NONE;
+	}
+	
+	PinnedSceneCapture->Modify();
+
+	UCompositePassBase* NewPass = NewObject<UCompositePassBase>(PinnedSceneCapture.Get(), InPassClass, NAME_None, RF_Transactional);
+
+	FProperty* PassListProperty = FindFProperty<FProperty>(UCompositeLayerSceneCapture::StaticClass(), GET_MEMBER_NAME_CHECKED(UCompositeLayerSceneCapture, LayerPasses));
+	PinnedSceneCapture->PreEditChange(PassListProperty);
+	
+	int32 NewPassIndex = SceneCapture->LayerPasses.Insert(NewPass, InPassIndex);
+
+	FPropertyChangedEvent PropertyChangedEvent(PassListProperty, EPropertyChangeType::ArrayAdd);
+	PinnedSceneCapture->PostEditChangeProperty(PropertyChangedEvent);
+
+	return NewPassIndex;
+}
+
+TArray<int32> FCompositeSceneCapturePassListOwner::CopyPasses(const TArray<UCompositePassBase*>& InPassesToCopy, int32 InGroupIndex, int32 InPassIndex)
+{
+	TArray<int32> OutNewPassIndices;
+	if (InGroupIndex != INDEX_NONE)
+	{
+		return OutNewPassIndices;
+	}
+
+	TStrongObjectPtr<UCompositeLayerSceneCapture> PinnedSceneCapture = SceneCapture.Pin();
+	if (!PinnedSceneCapture.IsValid())
+	{
+		return OutNewPassIndices;
+	}
+	
+	FProperty* PassListProperty = FindFProperty<FProperty>(UCompositeLayerSceneCapture::StaticClass(), GET_MEMBER_NAME_CHECKED(UCompositeLayerSceneCapture, LayerPasses));
+	PinnedSceneCapture->PreEditChange(PassListProperty);
+
+	int32 IndexToCopyTo = InPassIndex;
+	for (UCompositePassBase* PassToPaste : InPassesToCopy)
+	{
+		UCompositePassBase* NewPass = DuplicateObject(PassToPaste, PinnedSceneCapture.Get());
+			
+		const int32 NewPassIndex = PinnedSceneCapture->LayerPasses.Insert(NewPass, IndexToCopyTo);
+		OutNewPassIndices.Add(NewPassIndex);
+		++IndexToCopyTo;
+	}
+
+	FPropertyChangedEvent PropertyChangedEvent(PassListProperty, EPropertyChangeType::ArrayAdd);
+	PinnedSceneCapture->PostEditChangeProperty(PropertyChangedEvent);
+
+	return OutNewPassIndices;
+}
+
+void FCompositeSceneCapturePassListOwner::MovePass(int32 InSourceGroupIndex, int32 InSourcePassIndex, int32 InDestGroupIndex, int32 InDestPassIndex)
+{
+	if (InSourceGroupIndex != INDEX_NONE || InDestGroupIndex != INDEX_NONE)
+	{
+		return;
+	}
+
+	TStrongObjectPtr<UCompositeLayerSceneCapture> PinnedSceneCapture = SceneCapture.Pin();
+	if (!PinnedSceneCapture.IsValid())
+	{
+		return;
+	}
+
+	PinnedSceneCapture->Modify();
+
+	FProperty* PassListProperty = FindFProperty<FProperty>(UCompositeLayerSceneCapture::StaticClass(), GET_MEMBER_NAME_CHECKED(UCompositeLayerSceneCapture, LayerPasses));
+	UCompositePassBase* Pass = PinnedSceneCapture->LayerPasses[InSourcePassIndex];
+	
+	PinnedSceneCapture->PreEditChange(PassListProperty);
+	PinnedSceneCapture->LayerPasses.RemoveAt(InSourcePassIndex);
+	PinnedSceneCapture->LayerPasses.Insert(Pass, InDestPassIndex);
+
+	FPropertyChangedEvent MoveChangedEvent(PassListProperty, EPropertyChangeType::ArrayMove);
+	PinnedSceneCapture->PostEditChangeProperty(MoveChangedEvent);
+}
+
+void FCompositeSceneCapturePassListOwner::RemovePasses(int32 InGroupIndex, const TArray<int32>& InPassIndices)
+{
+	if (InGroupIndex != INDEX_NONE)
+	{
+		return;
+	}
+
+	TStrongObjectPtr<UCompositeLayerSceneCapture> PinnedSceneCapture = SceneCapture.Pin();
+	if (!PinnedSceneCapture.IsValid())
+	{
+		return;
+	}
+
+	// Delete the passes in reverse order so that indices don't get messed up as we delete
+	TArray<int32> SortedPassIndices = InPassIndices;
+	SortedPassIndices.Sort();
+
+	FProperty* PassListProperty = FindFProperty<FProperty>(UCompositeLayerSceneCapture::StaticClass(), GET_MEMBER_NAME_CHECKED(UCompositeLayerSceneCapture, LayerPasses));
+	PinnedSceneCapture->PreEditChange(PassListProperty);
+
+	for (int32 Index = SortedPassIndices.Num() - 1; Index >= 0; --Index)
+	{
+		const int32 PassIndex = SortedPassIndices[Index];
+
+		// Here we replicate details panels instanced property array removal code.
+		// The removed instanced pass objects must be moved to the transient package.
+		TObjectPtr<UCompositePassBase>& PassToRemove = PinnedSceneCapture->LayerPasses[PassIndex];
+		if (IsValid(PassToRemove))
+		{
+			PassToRemove->Modify();
+			PassToRemove->Rename(nullptr, GetTransientPackage(), REN_DontCreateRedirectors);
+		}
+
+		PinnedSceneCapture->LayerPasses.RemoveAt(PassIndex, EAllowShrinking::No);
+	}
+
+	PinnedSceneCapture->LayerPasses.Shrink();
+			
+	FPropertyChangedEvent PropertyChangedEvent(PassListProperty, EPropertyChangeType::ArrayAdd);
+	PinnedSceneCapture->PostEditChangeProperty(PropertyChangedEvent);
+}
diff --git a/Engine/Plugins/Compositing/Composite/Source/CompositeEditor/Private/UI/CompositeSceneCapturePassListOwner.h b/Engine/Plugins/Compositing/Composite/Source/CompositeEditor/Private/UI/CompositeSceneCapturePassListOwner.h
new file mode 100644
--- /dev/null
+++ b/Engine/Plugins/Compositing/Composite/Source/CompositeEditor/Private/UI/CompositeSceneCapturePassListOwner.h
@@ -0,0 +1,46 @@
+﻿// Copyright Epic Games, Inc. All Rights Reserved.
+
+#pragma once
+
+#include "SCompositePassTree.h"
+
+class UCompositeLayerSceneCapture;
+
+/** ICompositePassListOwner implementation for UCompositeLayerSceneCapture that manages its pass list */
+class FCompositeSceneCapturePassListOwner : public ICompositePassListOwner
+{
+public:
+	FCompositeSceneCapturePassListOwner(const TWeakObjectPtr<UCompositeLayerSceneCapture>& InSceneCapture)
+		: SceneCapture(InSceneCapture)
+	{ }
+
+	//~ ICompositePassListOwner interface
+	
+public:
+	virtual bool IsObjectValid() override;
+	virtual TStrongObjectPtr<UObject> GetObject() override;
+	virtual bool IsPassListPropertyName(const FName& InPropertyName) override;
+
+	virtual bool IsValidPassIndex(int32 InGroupIndex, int32 InPassIndex) override;
+	virtual UCompositePassBase* GetPass(int32 InGroupIndex, int32 InPassIndex) override;
+	virtual int32 GetNumGroups() const override { return INDEX_NONE; }
+	virtual TArray<TObjectPtr<UCompositePassBase>>& GetPassesForGroup(int32 InGroupIndex) override;
+	
+	virtual FString GetGroupFilterString(int32 InGroupIndex) override { return TEXT(""); }
+	virtual const FSlateBrush* GetGroupIcon(int32 InGroupIndex) override { return nullptr; }
+	virtual FText GetGroupDisplayName(int32 InGroupIndex) override { return FText::GetEmpty(); }
+	virtual FGroupFilterConfig GetGroupFilterConfig(int32 InGroupIndex) override { return FGroupFilterConfig(); }
+
+	virtual int32 GetDefaultGroupForNewPass(const UClass* InPassClass) const override { return INDEX_NONE; }
+	virtual bool CanAddPass(const UClass* InPassClass, int32 InGroupIndex, int32 InPassIndex) const override;
+	virtual int32 AddPass(const UClass* InPassClass, int32 InGroupIndex, int32 InPassIndex) override;
+	virtual TArray<int32> CopyPasses(const TArray<UCompositePassBase*>& InPassesToCopy, int32 InGroupIndex, int32 InPassIndex) override;
+	virtual void MovePass(int32 InSourceGroupIndex, int32 InSourcePassIndex, int32 InDestGroupIndex, int32 InDestPassIndex) override;
+	virtual void RemovePasses(int32 InGroupIndex, const TArray<int32>& InPassIndices) override;
+	
+	//~ EndICompositePassListOwner interface
+	
+private:
+	/** The scene capture UObject whose pass list is being provided */
+	TWeakObjectPtr<UCompositeLayerSceneCapture> SceneCapture;
+};
diff --git a/Engine/Plugins/Compositing/Composite/Source/CompositeEditor/Private/UI/SCompositeEditorPanel.cpp b/Engine/Plugins/Compositing/Composite/Source/CompositeEditor/Private/UI/SCompositeEditorPanel.cpp
--- a/Engine/Plugins/Compositing/Composite/Source/CompositeEditor/Private/UI/SCompositeEditorPanel.cpp
+++ b/Engine/Plugins/Compositing/Composite/Source/CompositeEditor/Private/UI/SCompositeEditorPanel.cpp
@@ -11,12 +11,14 @@
 #include "IDetailCustomization.h"
 #include "LevelEditor.h"
 #include "SCompositePanelLayerTree.h"
-#include "SCompositePlatePassPanel.h"
+#include "SCompositePassTreePanel.h"
 #include "ScopedTransaction.h"
 #include "WorkspaceMenuStructure.h"
 #include "WorkspaceMenuStructureModule.h"
 #include "Customizations/CompositeLayerPlateCustomization.h"
+#include "Customizations/CompositeLayerSceneCaptureCustomization.h"
 #include "Layers/CompositeLayerPlate.h"
+#include "Layers/CompositeLayerSceneCapture.h"
 #include "Widgets/Docking/SDockTab.h"
 #include "Widgets/Images/SImage.h"
 #include "Widgets/Input/SComboBox.h"
@@ -185,9 +187,11 @@ void SCompositeEditorPanel::Construct(const FArguments& InArgs)
 		return MakeShared<FCompositeActorPanelDetailCustomization>();
 	}));
 
-	// Create an instanced version of the plate customization so that a pointer to it can be stored and accessed later
+	// Create an instanced version of certain layer customizations so that a pointer to it can be stored and accessed later
 	DetailsView->RegisterInstancedCustomPropertyLayout(UCompositeLayerPlate::StaticClass(),
 		FOnGetDetailCustomizationInstance::CreateSP(this, &SCompositeEditorPanel::CreateCustomization<UCompositeLayerPlate, FCompositeLayerPlateCustomization>));
+	DetailsView->RegisterInstancedCustomPropertyLayout(UCompositeLayerSceneCapture::StaticClass(),
+		FOnGetDetailCustomizationInstance::CreateSP(this, &SCompositeEditorPanel::CreateCustomization<UCompositeLayerSceneCapture, FCompositeLayerSceneCaptureCustomization>));
 	
 	ChildSlot
 	[
@@ -257,43 +261,49 @@ void SCompositeEditorPanel::SelectObjects(const TArray<UObject*>& InSelection)
 		LayerTree->SelectObjects(CompositeActorsAndLayers);
 	}
 
-	for (const TPair<UCompositeLayerBase*, TArray<UCompositePassBase*>>& PassesForLayer : CompositePasses)
+	auto SelectPassesForClass = [this, &CompositePasses]<typename TClass, typename TCustomization>()
 	{
-		UCompositeLayerBase* Layer = PassesForLayer.Key;
-		TArray<UCompositePassBase*> Passes = PassesForLayer.Value;
-		
-		if (UCompositeLayerPlate* PlateOwner = Cast<UCompositeLayerPlate>(Layer))
+		for (const TPair<UCompositeLayerBase*, TArray<UCompositePassBase*>>& PassesForLayer : CompositePasses)
 		{
-			if (!InstancedCustomizations.Contains(UCompositeLayerPlate::StaticClass()))
+			UCompositeLayerBase* Layer = PassesForLayer.Key;
+			TArray<UCompositePassBase*> Passes = PassesForLayer.Value;
+		
+			if (TClass* TypedLayer = Cast<TClass>(Layer))
 			{
-				continue;
-			}
+				if (!InstancedCustomizations.Contains(TClass::StaticClass()))
+				{
+					continue;
+				}
 
-			TWeakPtr<IDetailCustomization>* Customization = InstancedCustomizations[UCompositeLayerPlate::StaticClass()].FindByPredicate([PlateOwner](const TWeakPtr<IDetailCustomization>& Customization)
-			{
-				if (TSharedPtr<IDetailCustomization> PinnedCustomization = Customization.Pin())
+				TWeakPtr<IDetailCustomization>* Customization = InstancedCustomizations[UCompositeLayerPlate::StaticClass()].FindByPredicate([TypedLayer](const TWeakPtr<IDetailCustomization>& Customization)
 				{
-					if (TSharedPtr<FCompositeLayerPlateCustomization> PlateCustomization = StaticCastSharedPtr<FCompositeLayerPlateCustomization>(PinnedCustomization))
+					if (TSharedPtr<IDetailCustomization> PinnedCustomization = Customization.Pin())
 					{
-						return PlateCustomization->IsCustomizingObject(PlateOwner);
+						if (TSharedPtr<TCustomization> TypeCustomization = StaticCastSharedPtr<TCustomization>(PinnedCustomization))
+						{
+							return TypeCustomization->IsCustomizingObject(TypedLayer);
+						}
 					}
-				}
 
-				return false;
-			});
+					return false;
+				});
 
-			if (Customization)
-			{
-				if (TSharedPtr<FCompositeLayerPlateCustomization> PlateCustomization = StaticCastSharedPtr<FCompositeLayerPlateCustomization>(Customization->Pin()))
+				if (Customization)
 				{
-					if (TSharedPtr<SCompositePlatePassPanel> PassPanel = PlateCustomization->GetPassPanelWidget())
+					if (TSharedPtr<TCustomization> PlateCustomization = StaticCastSharedPtr<TCustomization>(Customization->Pin()))
 					{
-						PassPanel->SelectPasses(Passes);
+						if (TSharedPtr<SCompositePassTreePanel> PassPanel = PlateCustomization->GetPassPanelWidget())
+						{
+							PassPanel->SelectPasses(Passes);
+						}
 					}
 				}
 			}
 		}
-	}
+	};
+	
+	SelectPassesForClass.operator()<UCompositeLayerPlate, FCompositeLayerPlateCustomization>();
+	SelectPassesForClass.operator()<UCompositeLayerSceneCapture, FCompositeLayerSceneCaptureCustomization>();
 }
 
 TArray<UObject*> SCompositeEditorPanel::GetSelectedObjects() const
@@ -304,27 +314,33 @@ TArray<UObject*> SCompositeEditorPanel::GetSelectedObjects() const
 		SelectedObjects.Append(LayerTree->GetSelectedObjects());
 	}
 
-	for (const TTuple<UClass*, TArray<TWeakPtr<IDetailCustomization>>>& CustomizationsForClass : InstancedCustomizations)
+	auto GetSelectedPassesForClass = [this, &SelectedObjects]<typename TClass, typename TCustomization>()
 	{
-		if (CustomizationsForClass.Key == UCompositeLayerPlate::StaticClass())
+		for (const TTuple<UClass*, TArray<TWeakPtr<IDetailCustomization>>>& CustomizationsForClass : InstancedCustomizations)
 		{
-			for (const TWeakPtr<IDetailCustomization>& Customization : CustomizationsForClass.Value)
+			if (CustomizationsForClass.Key == TClass::StaticClass())
 			{
-				if (!Customization.IsValid())
+				for (const TWeakPtr<IDetailCustomization>& Customization : CustomizationsForClass.Value)
 				{
-					continue;
-				}
+					if (!Customization.IsValid())
+					{
+						continue;
+					}
 
-				if (TSharedPtr<FCompositeLayerPlateCustomization> PlateCustomization = StaticCastSharedPtr<FCompositeLayerPlateCustomization>(Customization.Pin()))
-				{
-					if (TSharedPtr<SCompositePlatePassPanel> PassPanel = PlateCustomization->GetPassPanelWidget())
+					if (TSharedPtr<TCustomization> TypeCustomization = StaticCastSharedPtr<TCustomization>(Customization.Pin()))
 					{
-						SelectedObjects.Append(PassPanel->GetSelectedPasses());
+						if (TSharedPtr<SCompositePassTreePanel> PassPanel = TypeCustomization->GetPassPanelWidget())
+						{
+							SelectedObjects.Append(PassPanel->GetSelectedPasses());
+						}
 					}
 				}
 			}
 		}
-	}
+	};
+	
+	GetSelectedPassesForClass.operator()<UCompositeLayerPlate, FCompositeLayerPlateCustomization>();
+	GetSelectedPassesForClass.operator()<UCompositeLayerSceneCapture, FCompositeLayerSceneCaptureCustomization>();
 	
 	return SelectedObjects;
 }
diff --git a/Engine/Plugins/Compositing/Composite/Source/CompositeEditor/Private/UI/SCompositePassTree.cpp b/Engine/Plugins/Compositing/Composite/Source/CompositeEditor/Private/UI/SCompositePassTree.cpp
--- a/Engine/Plugins/Compositing/Composite/Source/CompositeEditor/Private/UI/SCompositePassTree.cpp
+++ b/Engine/Plugins/Compositing/Composite/Source/CompositeEditor/Private/UI/SCompositePassTree.cpp
@@ -18,7 +18,6 @@
 #include "Layers/CompositeLayerPlate.h"
 #include "Misc/StringOutputDevice.h"
 #include "Passes/CompositePassDistortion.h"
-#include "Passes/CompositePassMaterial.h"
 #include "Styling/SlateIconFinder.h"
 #include "UObject/Package.h"
 #include "Widgets/Text/SInlineEditableTextBlock.h"
@@ -159,8 +158,9 @@ public:
 		SLATE_EVENT(FOnFilterNewPassType, OnFilterNewPassType)
 	SLATE_END_ARGS()
 
-	void Construct(const FArguments& InArgs)
+	void Construct(const FArguments& InArgs, const TSharedPtr<ICompositePassListOwner>& InPassListOwner)
 	{
+		PassListOwner = InPassListOwner;
 		OnFilterChanged = InArgs._OnFilterChanged;
 		OnPassAdded = InArgs._OnPassAdded;
 		OnFilterNewPassType = InArgs._OnFilterNewPassType;
@@ -172,7 +172,10 @@ public:
 			.CustomFilters(AllFilters)
 			.OnFilterChanged(this, &SCompositePassTreeToolbar::FilterChanged);
 
-		AllPassesFilter->SetActive(true);
+		if (AllPassesFilter.IsValid())
+		{
+			AllPassesFilter->SetActive(true);
+		}
 		
 		ChildSlot
 		[
@@ -245,12 +248,16 @@ public:
 	{
 		// Prevent filter changed callback from firing while changing filter status, will invoke once all filters are reset
 		TGuardValue<bool> FilterChangedGuard(bFilterChangedGuard, true);
-		for (const TSharedPtr<FCompositePassTreeFilter>& Filter : PassTypeFilters)
+		for (const TSharedPtr<FCompositePassTreeFilter>& Filter : GroupFilters)
 		{
 			Filter->DeactivateFilter();
 		}
 
-		AllPassesFilter->ActivateFilter();
+		if (AllPassesFilter.IsValid())
+		{
+			AllPassesFilter->ActivateFilter();
+		}
+
 		TextFilterSearchBox->SetText(FText::GetEmpty());
 
 		OnFilterChanged.ExecuteIfBound();
@@ -262,49 +269,47 @@ private:
 	{
 		TSharedPtr<FFilterCategory> BasicFiltersCategory = MakeShared<FFilterCategory>(LOCTEXT("BasicFiltersCategory", "Basic"), FText::GetEmpty());
 
-		static TArray<TTuple<FString, FText, FText>> PassTypeFilterConfigs =
-		{
-			{ TEXT("MediaPassesFilter"), LOCTEXT("MediaPassesFilterName", "Media Passes"),  LOCTEXT("MediaPassesFilterToolTip", "Only show media passes") },
-			//{ TEXT("ScenePassesFilter"), LOCTEXT("ScenePassesFilterName", "Scene Passes"),  LOCTEXT("ScenePassesFilterToolTip", "Only show scene passes")  },
-			{ TEXT("LayerPassesFilter"), LOCTEXT("LayerPassesFilterName", "Layer Passes"),  LOCTEXT("LayerPassesFilterToolTip", "Only show layer passes")  }
-		};
-
-		for (int32 Index = 0; Index < PassTypeFilterConfigs.Num(); ++Index)
+		const int32 NumGroups = PassListOwner.IsValid() ? PassListOwner->GetNumGroups() : 0;
+		for (int32 GroupIndex = 0; GroupIndex < NumGroups; ++GroupIndex)
 		{
+			ICompositePassListOwner::FGroupFilterConfig FilterConfig = PassListOwner->GetGroupFilterConfig(GroupIndex);
 			TSharedPtr<FCompositePassTreeFilter> PassTypeFilter = MakeShared<FCompositePassTreeFilter>(
 				BasicFiltersCategory,
-				PassTypeFilterConfigs[Index].Get<0>(),
-				PassTypeFilterConfigs[Index].Get<1>(),
-				FGenericFilter<FFilterType>::FOnItemFiltered::CreateLambda([Index](FFilterType InItem)
+				FilterConfig.FilterName,
+				FilterConfig.DisplayName,
+				FGenericFilter<FFilterType>::FOnItemFiltered::CreateLambda([GroupIndex](FFilterType InItem)
 				{
 					if (!InItem.IsValid())
 					{
 						return false;
 					}
 
-					return (int32)InItem->PassType == Index;
+					return InItem->GroupIndex == GroupIndex;
 				}));
 	
-			PassTypeFilter->SetToolTipText(PassTypeFilterConfigs[Index].Get<2>());
-			PassTypeFilter->OnActiveStateChanged.BindSP(this, &SCompositePassTreeToolbar::OnPassTypeFilterStateChanged, Index);
+			PassTypeFilter->SetToolTipText(FilterConfig.ToolTip);
+			PassTypeFilter->OnActiveStateChanged.BindSP(this, &SCompositePassTreeToolbar::OnGroupFilterStateChanged, GroupIndex);
 
-			PassTypeFilters.Add(PassTypeFilter);
+			GroupFilters.Add(PassTypeFilter);
 			OutAllFilters.Add(PassTypeFilter.ToSharedRef());
 		}
-		
-		AllPassesFilter = MakeShared<FCompositePassTreeFilter>(
-				BasicFiltersCategory,
-				TEXT("AllPassesFilter"),
-				LOCTEXT("AllPassesFilterName", "All Passes"),
-				FGenericFilter<FFilterType>::FOnItemFiltered::CreateLambda([](FFilterType InItem)
-				{
-					return true;
-				}));
+
+		if (NumGroups > 0)
+		{
+			AllPassesFilter = MakeShared<FCompositePassTreeFilter>(
+					BasicFiltersCategory,
+					TEXT("AllPassesFilter"),
+					LOCTEXT("AllPassesFilterName", "All Passes"),
+					FGenericFilter<FFilterType>::FOnItemFiltered::CreateLambda([](FFilterType InItem)
+					{
+						return true;
+					}));
 	
-		AllPassesFilter->SetToolTipText(LOCTEXT("AllPassesFilterToolTip", "Show all passes"));
-		AllPassesFilter->OnActiveStateChanged.BindSP(this, &SCompositePassTreeToolbar::OnAllPassesFilterStateChanged);
+			AllPassesFilter->SetToolTipText(LOCTEXT("AllPassesFilterToolTip", "Show all passes"));
+			AllPassesFilter->OnActiveStateChanged.BindSP(this, &SCompositePassTreeToolbar::OnAllPassesFilterStateChanged);
 		
-		OutAllFilters.Add(AllPassesFilter.ToSharedRef());
+			OutAllFilters.Add(AllPassesFilter.ToSharedRef());
+		}
 		
 		TreeTextFilter = MakeShared<FTreeFilter>(FTreeFilter::FItemToStringArray::CreateLambda([](const FFilterType& InItem, TArray<FString>& OutStrings)
 		{
@@ -322,14 +327,11 @@ private:
 					}
 				}
 
-				static FString PassTypeFilterStrings[] =
+				const FString GroupFilterString = InItem->GetGroupFilterString();
+				if (!GroupFilterString.IsEmpty())
 				{
-					TEXT("Media Passes"),
-					// TEXT("Scene Passes"),
-					TEXT("Layer Passes")
-				};
-				
-				OutStrings.Add(PassTypeFilterStrings[(int32)InItem->PassType]);
+					OutStrings.Add(GroupFilterString);
+				}
 			}
 		}));
 		TreeTextFilter->OnChanged().AddSP(this, &SCompositePassTreeToolbar::FilterChanged);
@@ -420,26 +422,26 @@ private:
 		OnPassAdded.ExecuteIfBound(InPassClass);	
 	}
 
-	void OnPassTypeFilterStateChanged(bool bActive, int32 PassTypeIndex)
+	void OnGroupFilterStateChanged(bool bActive, int32 GroupIndex)
 	{
 		if (bActive)
 		{
-			for (int32 Index = 0; Index < PassTypeFilters.Num(); ++Index)
+			for (int32 Index = 0; Index < GroupFilters.Num(); ++Index)
 			{
-				if (Index != PassTypeIndex && PassTypeFilters[Index]->IsActive())
+				if (Index != GroupIndex && GroupFilters[Index]->IsActive())
 				{
-					PassTypeFilters[Index]->SetActive(false);
+					GroupFilters[Index]->SetActive(false);
 				}
 			}
 			
-			if (AllPassesFilter->IsActive())
+			if (AllPassesFilter.IsValid() && AllPassesFilter->IsActive())
 			{
 				AllPassesFilter->DeactivateFilter();
 			}
 		}
 		else
 		{
-			if (!AllPassesFilter->IsActive())
+			if (AllPassesFilter.IsValid() && !AllPassesFilter->IsActive())
 			{
 				AllPassesFilter->ActivateFilter();
 			}
@@ -450,7 +452,7 @@ private:
 	{
 		if (bActive)
 		{
-			for (TSharedPtr<FCompositePassTreeFilter>& Filter : PassTypeFilters)
+			for (TSharedPtr<FCompositePassTreeFilter>& Filter : GroupFilters)
 			{
 				if (Filter->IsActive())
 				{
@@ -460,9 +462,12 @@ private:
 		}
 		else
 		{
-			if (!PassTypeFilters[0]->IsActive())
+			if (!GroupFilters.IsEmpty())
 			{
-				PassTypeFilters[0]->ActivateFilter();
+				if (!GroupFilters[0]->IsActive())
+				{
+					GroupFilters[0]->ActivateFilter();
+				}
 			}
 		}
 	}
@@ -478,8 +483,10 @@ private:
 	}
 	
 private:
+	TSharedPtr<ICompositePassListOwner> PassListOwner;
+	
 	TSharedPtr<FCompositePassTreeFilter> AllPassesFilter;
-	TArray<TSharedPtr<FCompositePassTreeFilter>> PassTypeFilters;
+	TArray<TSharedPtr<FCompositePassTreeFilter>> GroupFilters;
 	TSharedPtr<FTreeFilter> TreeTextFilter;
 	
 	TSharedPtr<SFilterSearchBox> TextFilterSearchBox;
@@ -500,7 +507,7 @@ class SCompositePassTreeItemRow : public SMultiColumnTableRow<SCompositePassTree
 	using FTreeItemWeakPtr = TWeakPtr<FTreeItem>;
 	
 public:
-	DECLARE_DELEGATE_ThreeParams(FOnPassMoved, const FTreeItemPtr& /* InItemToMove */, SCompositePassTree::EPassType /* InDestPassType */, int32 /* InDestIndex */)
+	DECLARE_DELEGATE_ThreeParams(FOnPassMoved, const FTreeItemPtr& /* InItemToMove */, int32 /* InDestGroupIndex */, int32 /* InDestIndex */)
 
 private:
 	/** Drag/drop operation for changing the order of passes in the tree view */
@@ -629,7 +636,7 @@ private:
 	/** Gets the checkbox state for the row item's Enabled flag */
 	ECheckBoxState IsItemEnabled() const
 	{
-		if (!TreeItem->Plate.IsValid())
+		if (!TreeItem->HasValidOwner())
 		{
 			return ECheckBoxState::Unchecked;
 		}
@@ -645,25 +652,27 @@ private:
 			return ECheckBoxState::Unchecked;
 		}
 		
-		// Tree item is a pass type group item, so set its checkbox state as the aggregate of all passes of that type
-		TArray<TObjectPtr<UCompositePassBase>>& AllPassesOfType = SCompositePassTree::GetPassList(TreeItem->Plate, TreeItem->PassType);
+		// Tree item is a pass type group item, so set its checkbox state as the aggregate of all its children
 
 		bool bAnyEnabled = false;
 		bool bAnyDisabled = false;
-		for (const TObjectPtr<UCompositePassBase>& Pass : AllPassesOfType)
+		for (const FTreeItemPtr& ChildTreeItem : TreeItem->Children)
 		{
-			if (!IsValid(Pass))
+			if (!ChildTreeItem.IsValid() || !ChildTreeItem->HasValidOwner() || !ChildTreeItem->HasValidPassIndex())
 			{
 				continue;
 			}
-			
-			if (Pass->IsEnabled())
-			{
-				bAnyEnabled = true;
-			}
-			else
+
+			if (UCompositePassBase* Pass = ChildTreeItem->GetPass())
 			{
-				bAnyDisabled = true;
+				if (Pass->IsEnabled())
+				{
+					bAnyEnabled = true;
+				}
+				else
+				{
+					bAnyDisabled = true;
+				}
 			}
 		}
 
@@ -680,7 +689,7 @@ private:
 	/** Raised when the row item's Enabled checkbox has changed state */
 	void OnIsItemEnabledChanged(ECheckBoxState InCheckBoxState)
 	{
-		if (!TreeItem->Plate.IsValid())
+		if (!TreeItem->HasValidOwner())
 		{
 			return;
 		}
@@ -700,25 +709,27 @@ private:
 			return;
 		}
 		
-		// Tree item is a pass type group item, so set enable state of all passes of type
-		TArray<TObjectPtr<UCompositePassBase>>& AllPassesOfType = SCompositePassTree::GetPassList(TreeItem->Plate, TreeItem->PassType);
-		if (AllPassesOfType.IsEmpty())
-		{
-			return;
-		}
-		
+		// Tree item is a group item, so set enable state of all child passes
 		FScopedTransaction SetAllEnabledTransaction(LOCTEXT("SetAllEnabledTransaction", "Set All Enabled"));
-		for (const TObjectPtr<UCompositePassBase>& Pass : AllPassesOfType)
+		for (const FTreeItemPtr& ChildTreeItem : TreeItem->Children)
 		{
-			Pass->Modify();
-			Pass->SetEnabled(InCheckBoxState == ECheckBoxState::Checked);
+			if (!ChildTreeItem.IsValid() || !ChildTreeItem->HasValidOwner() || !ChildTreeItem->HasValidPassIndex())
+			{
+				continue;
+			}
+
+			if (UCompositePassBase* Pass = ChildTreeItem->GetPass())
+			{
+				Pass->Modify();
+				Pass->SetEnabled(InCheckBoxState == ECheckBoxState::Checked);
+			}
 		}
 	}
 	
 	/** Gets the row item's icon to display next to its name */
 	const FSlateBrush* GetItemIcon() const
 	{
-		if (!TreeItem->Plate.IsValid())
+		if (!TreeItem->HasValidOwner())
 		{
 			return nullptr;
 		}
@@ -734,21 +745,14 @@ private:
 			return nullptr;
 		}
 
-		// Tree item is a pass type group item, so return icon for pass type
-		static const FSlateBrush* PassTypeIcons[] =
-		{
-			FCompositeEditorStyle::Get().GetBrush("CompositeEditor.Passes.Media"),
-			//FCompositeEditorStyle::Get().GetBrush("CompositeEditor.Passes.Scene"),
-			FCompositeEditorStyle::Get().GetBrush("CompositeEditor.Passes.Layer")
-		};
-		
-		return PassTypeIcons[(int32)TreeItem->PassType];
+		// Tree item is a group item, so return icon for the group
+		return TreeItem->GetGroupIcon();
 	}
 
 	/** Gets the row item's display name */
 	FText GetItemName() const
 	{
-		if (!TreeItem->Plate.IsValid())
+		if (!TreeItem->HasValidOwner())
 		{
 			return FText::GetEmpty();
 		}
@@ -764,21 +768,14 @@ private:
 			return LOCTEXT("PassNotSetLabel", "Pass not set");
 		}
 
-		// Tree item is a pass type group item, so return label for pass type
-		static FText PassTypeLabels[] =
-		{
-			LOCTEXT("MediaPassTypeLabel", "Media Passes"),
-			//LOCTEXT("ScenePassTypeLabel", "Scene Passes"),
-			LOCTEXT("LayerPassTypeLabel", "Layer Passes"),
-		};
-		
-		return PassTypeLabels[(int32)TreeItem->PassType];
+		// Tree item is a group item, so return the label for the group
+		return TreeItem->GetGroupDisplayName();
 	}
 
 	/** Sets this row item's display name */
 	void SetItemName(const FText& InNewText, ETextCommit::Type InCommitType)
 	{
-		if (!TreeItem->Plate.IsValid())
+		if (!TreeItem->HasValidOwner())
 		{
 			return;
 		}
@@ -798,7 +795,7 @@ private:
 	/** Gets whether the row item's display name can be changed */
 	bool IsItemNameReadOnly() const
 	{
-		if (!TreeItem->Plate.IsValid())
+		if (!TreeItem->HasValidOwner())
 		{
 			return true;
 		}
@@ -817,7 +814,7 @@ private:
 	/** Gets the row item's display name */
 	FText GetItemType() const
 	{
-		if (!TreeItem->Plate.IsValid())
+		if (!TreeItem->HasValidOwner())
 		{
 			return FText::GetEmpty();
 		}
@@ -839,7 +836,7 @@ private:
 	/** Raised when a drag start has been detected on this row item */
 	FReply HandleDragDetected(const FGeometry& Geometry, const FPointerEvent& MouseEvent)
 	{
-		if (!TreeItem->Plate.IsValid())
+		if (!TreeItem->HasValidOwner())
 		{
 			return FReply::Unhandled();
 		}
@@ -869,12 +866,12 @@ private:
 		}
 
 		FTreeItemPtr DraggedItem = DragDropOp->DraggedItem.Pin();
-		if (DraggedItem->Plate.IsValid())
+		if (DraggedItem->HasValidOwner() && InTreeItem->HasValidOwner())
 		{
 			if (UCompositePassBase* Pass = DraggedItem->GetPass())
 			{
-				// PPM passes can only be added to the Layer passes list
-				if (Pass->GetClass()->IsChildOf<UCompositePassMaterial>() && InTreeItem->PassType != SCompositePassTree::EPassType::Layer)
+				const int32 DestIndex = InItemDropZone == EItemDropZone::AboveItem ? InTreeItem->PassIndex: InTreeItem->PassIndex + 1;
+				if (!InTreeItem->Owner->CanAddPass(Pass->GetClass(), InTreeItem->GroupIndex, DestIndex))
 				{
 					return TOptional<EItemDropZone>();
 				}
@@ -907,12 +904,12 @@ private:
 
 			// If the destination is further down than the item being moved, we must subtract one from the destination index
 			// to account for the removal of the original item
-			if (InTreeItem->PassIndex > DragDropOp->DraggedItem.Pin()->PassIndex && InTreeItem->PassType == DragDropOp->DraggedItem.Pin()->PassType)
+			if (InTreeItem->PassIndex > DragDropOp->DraggedItem.Pin()->PassIndex && InTreeItem->GroupIndex == DragDropOp->DraggedItem.Pin()->GroupIndex)
 			{
 				--DestIndex;
 			}
 			
-			OnPassMoved.Execute(DragDropOp->DraggedItem.Pin(), InTreeItem->PassType, DestIndex);
+			OnPassMoved.Execute(DragDropOp->DraggedItem.Pin(), InTreeItem->GroupIndex, DestIndex);
 			return FReply::Handled();
 		}
 		
@@ -930,15 +927,19 @@ private:
 	FOnPassMoved OnPassMoved;
 };
 
+bool SCompositePassTree::FPassTreeItem::HasValidOwner() const
+{
+	return Owner.IsValid() && Owner->IsObjectValid();
+}
+
 bool SCompositePassTree::FPassTreeItem::HasValidPassIndex() const
 {
-	if (!Plate.IsValid())
+	if (!Owner.IsValid())
 	{
 		return false;
 	}
 
-	TArray<TObjectPtr<UCompositePassBase>>& PassList = SCompositePassTree::GetPassList(Plate, PassType);
-	return PassList.IsValidIndex(PassIndex);
+	return Owner->IsValidPassIndex(GroupIndex, PassIndex);
 }
 
 UCompositePassBase* SCompositePassTree::FPassTreeItem::GetPass() const
@@ -948,13 +949,42 @@ UCompositePassBase* SCompositePassTree::FPassTreeItem::GetPass() const
 		return nullptr;
 	}
 
-	TArray<TObjectPtr<UCompositePassBase>>& PassList = SCompositePassTree::GetPassList(Plate, PassType);
-	return PassList[PassIndex];
+	return Owner->GetPass(GroupIndex, PassIndex);
+}
+
+FString SCompositePassTree::FPassTreeItem::GetGroupFilterString() const
+{
+	if (!Owner.IsValid())
+	{
+		return TEXT("");
+	}
+
+	return Owner->GetGroupFilterString(GroupIndex);
 }
 
-void SCompositePassTree::Construct(const FArguments& InArgs, const TWeakObjectPtr<UCompositeLayerPlate>& InPlate)
+const FSlateBrush* SCompositePassTree::FPassTreeItem::GetGroupIcon() const
 {
-	Plate = InPlate;
+	if (!Owner.IsValid())
+	{
+		return nullptr;
+	}
+
+	return Owner->GetGroupIcon(GroupIndex);
+}
+
+FText SCompositePassTree::FPassTreeItem::GetGroupDisplayName() const
+{
+	if (!Owner.IsValid())
+	{
+		return FText::GetEmpty();
+	}
+
+	return Owner->GetGroupDisplayName(GroupIndex);
+}
+
+void SCompositePassTree::Construct(const FArguments& InArgs, const TSharedPtr<ICompositePassListOwner>& InPassListOwner)
+{
+	PassListOwner = InPassListOwner;
 	OnSelectionChanged = InArgs._OnSelectionChanged;
 	OnLayoutChanged = InArgs._OnLayoutChanged;
 
@@ -971,7 +1001,7 @@ void SCompositePassTree::Construct(const FArguments& InArgs, const TWeakObjectPt
 		.AutoHeight()
 		.Padding(0.0f, 4.0f)
 		[
-			SAssignNew(Toolbar, SCompositePassTreeToolbar)
+			SAssignNew(Toolbar, SCompositePassTreeToolbar, PassListOwner)
 			.OnFilterChanged(this, &SCompositePassTree::OnFilterChanged)
 			.OnPassAdded(this, &SCompositePassTree::OnPassAdded)
 			.OnFilterNewPassType(this, &SCompositePassTree::OnFilterNewPassType)
@@ -1057,6 +1087,17 @@ FReply SCompositePassTree::OnKeyDown(const FGeometry& MyGeometry, const FKeyEven
 	return FReply::Unhandled();
 }
 
+void SCompositePassTree::PostUndo(bool bSuccess)
+{
+	constexpr bool bPreserveSelection = true;
+	FillPassTreeItems(bPreserveSelection);
+}
+
+void SCompositePassTree::PostRedo(bool bSuccess)
+{
+	PostUndo(bSuccess);
+}
+
 void SCompositePassTree::SelectPasses(const TArray<UCompositePassBase*>& InPasses)
 {
 	if (!TreeView.IsValid())
@@ -1112,40 +1153,6 @@ TArray<UCompositePassBase*> SCompositePassTree::GetSelectedPasses() const
 	return SelectedPasses;
 }
 
-TArray<TObjectPtr<UCompositePassBase>>& SCompositePassTree::GetPassList(TWeakObjectPtr<UCompositeLayerPlate> InPlate, EPassType InPassType)
-{
-	check(InPlate.IsValid());
-	
-	switch (InPassType)
-	{
-	case EPassType::Media:
-		return InPlate->MediaPasses;
-
-	// case EPassType::Scene:
-	// return InPlate->ScenePasses;
-		
-	case EPassType::Layer:
-	default:
-		return InPlate->LayerPasses;
-	}
-}
-
-FProperty* SCompositePassTree::GetPassListProperty(EPassType InPassType)
-{
-	switch (InPassType)
-	{
-	case EPassType::Media:
-		return FindFProperty<FProperty>(UCompositeLayerPlate::StaticClass(), GET_MEMBER_NAME_CHECKED(UCompositeLayerPlate, MediaPasses));
-
-	// case EPassType::Scene:
-	// return FindFProperty<FProperty>(UCompositeLayerPlate::StaticClass(), GET_MEMBER_NAME_CHECKED(UCompositeLayerPlate, ScenePasses));
-		
-	case EPassType::Layer:
-	default:
-		return FindFProperty<FProperty>(UCompositeLayerPlate::StaticClass(), GET_MEMBER_NAME_CHECKED(UCompositeLayerPlate, LayerPasses));
-	}
-}
-
 void SCompositePassTree::BindCommands()
 {
 	CommandList->MapAction(FGenericCommands::Get().Copy,
@@ -1177,36 +1184,132 @@ void SCompositePassTree::BindCommands()
 		FCanExecuteAction::CreateSP(this, &SCompositePassTree::CanEnableSelectedItems));
 }
 
-void SCompositePassTree::FillPassTreeItems()
+void SCompositePassTree::FillPassTreeItems(bool bPreserveSelection)
 {
+	// A list of passes to reselect after the list of pass tree items has been regenerated
+	TArray<UCompositePassBase*> PassesToReselect;
+	if (bPreserveSelection)
+	{
+		TArray<FPassTreeItemPtr> SelectedItems = TreeView.IsValid() ? TreeView->GetSelectedItems() : TArray<FPassTreeItemPtr>();
+		for (const FPassTreeItemPtr& Item : SelectedItems)
+		{
+			if (Item.IsValid() && Item->HasValidPassIndex())
+			{
+				if (UCompositePassBase* Pass = Item->GetPass())
+				{
+					PassesToReselect.Add(Pass);
+				}
+			}
+		}
+	}
+	
 	PassTreeItems.Empty();
 
-	if (Plate.IsValid())
+	if (PassListOwner.IsValid() && PassListOwner->IsObjectValid())
 	{
-		for (int32 TypeIndex = 0; TypeIndex < (int32)EPassType::MAX; ++TypeIndex)
+		const int32 NumGroups = PassListOwner->GetNumGroups();
+		if (NumGroups <= 0)
 		{
-			const EPassType PassType = (EPassType)TypeIndex;
-			FPassTreeItemPtr NewPassTypeTreeItem = MakeShared<FPassTreeItem>();
-			NewPassTypeTreeItem->Plate = Plate;
-			NewPassTypeTreeItem->PassType = PassType;
-			
-			PassTreeItems.Add(NewPassTypeTreeItem);
-			
-			TArray<TObjectPtr<UCompositePassBase>>& Passes = GetPassList(PassType);
+			TArray<TObjectPtr<UCompositePassBase>>& Passes = PassListOwner->GetPassesForGroup(INDEX_NONE);
 			for (int32 Index = 0; Index < Passes.Num(); ++Index)
 			{
 				FPassTreeItemPtr NewPassTreeItem = MakeShared<FPassTreeItem>();
-				NewPassTreeItem->Plate = Plate;
-				NewPassTreeItem->PassType = PassType;
+				NewPassTreeItem->Owner = PassListOwner;
+				NewPassTreeItem->GroupIndex = INDEX_NONE;
 				NewPassTreeItem->PassIndex = Index;
 
-				NewPassTypeTreeItem->Children.Add(NewPassTreeItem);
+				PassTreeItems.Add(NewPassTreeItem);
+			}
+		}
+		else
+		{
+			for (int32 GroupIndex = 0; GroupIndex < NumGroups; ++GroupIndex)
+			{
+				FPassTreeItemPtr NewPassTypeTreeItem = MakeShared<FPassTreeItem>();
+				NewPassTypeTreeItem->Owner = PassListOwner;
+				NewPassTypeTreeItem->GroupIndex = GroupIndex;
+			
+				PassTreeItems.Add(NewPassTypeTreeItem);
+			
+				TArray<TObjectPtr<UCompositePassBase>>& Passes = PassListOwner->GetPassesForGroup(GroupIndex);
+				for (int32 Index = 0; Index < Passes.Num(); ++Index)
+				{
+					FPassTreeItemPtr NewPassTreeItem = MakeShared<FPassTreeItem>();
+					NewPassTreeItem->Owner = PassListOwner;
+					NewPassTreeItem->GroupIndex = GroupIndex;
+					NewPassTreeItem->PassIndex = Index;
+
+					NewPassTypeTreeItem->Children.Add(NewPassTreeItem);
+				}
 			}
 		}
 	}
 
 	FilterPassTreeItems();
 	RefreshAndExpandTreeView();
+
+	if (TreeView.IsValid())
+	{
+		// Silence any selection changes while re-selecting, and track if the selection has changed (i.e. the originally selected items no longer exist)
+		// If the selection has changed at the end of re-selection, the selection changed delegate will be invoked manually
+		bool bSelectionChanged = false;
+		{
+			TGuardValue<bool> SilenceSelectionChangesGuard(bSilenceSelectionChanges, true);
+
+			auto FindTreeItemForPass = [this](const UCompositePassBase* Pass)->FPassTreeItemPtr
+			{
+				for (const FPassTreeItemPtr& TreeItem : FilteredPassTreeItems)
+				{
+					if (!TreeItem.IsValid() || !TreeItem->HasValidOwner())
+					{
+						continue;
+					}
+
+					if (TreeItem->HasValidPassIndex())
+					{
+						if (TreeItem->GetPass() == Pass)
+						{
+							return TreeItem;
+						}
+
+						continue;
+					}
+					
+					for (const FPassTreeItemPtr& ChildItem : TreeItem->Children)
+					{
+						if (ChildItem->HasValidPassIndex())
+						{
+							if (ChildItem->GetPass() == Pass)
+							{
+								return ChildItem;
+							}
+						}
+					}
+				}
+
+				return nullptr;
+			};
+			
+			for (const UCompositePassBase* PassToReselect : PassesToReselect)
+			{
+				FPassTreeItemPtr TreeItemToReselect = FindTreeItemForPass(PassToReselect);
+
+				if (!TreeItemToReselect.IsValid())
+				{
+					// Pass no longer exists in the displayed list of pass items. Invoke selection changed to clear out selection from other UI elements
+					bSelectionChanged = true;
+					continue;
+				}
+
+				TreeView->SetItemSelection(TreeItemToReselect, true);
+			}
+		}
+
+		if (bSelectionChanged)
+		{
+			OnPassSelectionChanged(nullptr, ESelectInfo::Direct);
+		}
+	}
 }
 
 void SCompositePassTree::FilterPassTreeItems()
@@ -1239,24 +1342,51 @@ void SCompositePassTree::FilterPassTreeItems()
 	}
 }
 
-void SCompositePassTree::RefreshPassType(EPassType InPassType, bool bRefreshTreeView)
+void SCompositePassTree::RefreshGroup(int32 InGroupIndex, bool bRefreshTreeView)
 {
-	FPassTreeItemPtr PassTypeTreeItem = PassTreeItems[(int32)InPassType];
-	PassTypeTreeItem->Children.Empty();
-
-	TArray<TObjectPtr<UCompositePassBase>>& Passes = GetPassList(InPassType);
-	for (int32 Index = 0; Index < Passes.Num(); ++Index)
+	if (InGroupIndex == INDEX_NONE)
 	{
-		FPassTreeItemPtr NewPassTreeItem = MakeShared<FPassTreeItem>();
-		NewPassTreeItem->Plate = Plate;
-		NewPassTreeItem->PassType = InPassType;
-		NewPassTreeItem->PassIndex = Index;
+		// Remove any pass tree items that aren't children of a group
+		PassTreeItems.RemoveAll([](const FPassTreeItemPtr& TreeItem)
+		{
+			return !TreeItem.IsValid() || TreeItem->GroupIndex == INDEX_NONE;
+		});
 
-		PassTypeTreeItem->Children.Add(NewPassTreeItem);
+		TArray<TObjectPtr<UCompositePassBase>>& Passes = PassListOwner->GetPassesForGroup(INDEX_NONE);
+		for (int32 Index = 0; Index < Passes.Num(); ++Index)
+		{
+			FPassTreeItemPtr NewPassTreeItem = MakeShared<FPassTreeItem>();
+			NewPassTreeItem->Owner = PassListOwner;
+			NewPassTreeItem->GroupIndex = INDEX_NONE;
+			NewPassTreeItem->PassIndex = Index;
+
+			PassTreeItems.Add(NewPassTreeItem);
+		}
 	}
+	else
+	{
+		FPassTreeItemPtr PassTypeTreeItem = PassTreeItems[InGroupIndex];
+		PassTypeTreeItem->Children.Empty();
 
-	FilterPassTreeItems();
+		if (!PassListOwner.IsValid() || !PassListOwner->IsObjectValid())
+		{
+			return;
+		}
+	
+		TArray<TObjectPtr<UCompositePassBase>>& Passes = PassListOwner->GetPassesForGroup(InGroupIndex);
+		for (int32 Index = 0; Index < Passes.Num(); ++Index)
+		{
+			FPassTreeItemPtr NewPassTreeItem = MakeShared<FPassTreeItem>();
+			NewPassTreeItem->Owner = PassListOwner;
+			NewPassTreeItem->GroupIndex = InGroupIndex;
+			NewPassTreeItem->PassIndex = Index;
+
+			PassTypeTreeItem->Children.Add(NewPassTreeItem);
+		}
+	}
 	
+	FilterPassTreeItems();
+
 	if (bRefreshTreeView)
 	{
 		RefreshAndExpandTreeView();
@@ -1310,7 +1440,7 @@ TSharedPtr<SWidget> SCompositePassTree::CreateTreeContextMenu()
 
 ECheckBoxState SCompositePassTree::GetGlobalEnabledState() const
 {
-	if (!Plate.IsValid())
+	if (!PassListOwner.IsValid() || !PassListOwner->IsObjectValid())
 	{
 		return ECheckBoxState::Checked;
 	}
@@ -1318,10 +1448,10 @@ ECheckBoxState SCompositePassTree::GetGlobalEnabledState() const
 	// Determine if all passes are disabled, enabled, or if there is a mix
 	bool bAnyPassesEnabled = false;
 	bool bAnyPassesDisabled = false;
-	
-	for (int32 Index = 0; Index < (int32)EPassType::MAX; ++Index)
+
+	auto CheckPassListEnabledStatus = [this, &bAnyPassesEnabled, &bAnyPassesDisabled](int32 InGroupIndex)
 	{
-		const TArray<TObjectPtr<UCompositePassBase>>& PassList = GetPassList((EPassType)Index);
+		const TArray<TObjectPtr<UCompositePassBase>>& PassList = PassListOwner->GetPassesForGroup(InGroupIndex);
 
 		// Consider any pass list that is empty as enabled to align with how the check state of the pass type tree item works
 		if (PassList.IsEmpty())
@@ -1345,6 +1475,19 @@ ECheckBoxState SCompositePassTree::GetGlobalEnabledState() const
 				bAnyPassesDisabled = true;
 			}
 		}
+	};
+	
+	const int32 NumGroups = PassListOwner->GetNumGroups();
+	if (NumGroups <= 0)
+	{
+		CheckPassListEnabledStatus(INDEX_NONE);
+	}
+	else
+	{
+		for (int32 GroupIndex = 0; GroupIndex < NumGroups; ++GroupIndex)
+		{
+			CheckPassListEnabledStatus(GroupIndex);
+		}
 	}
 
 	if (bAnyPassesEnabled)
@@ -1359,83 +1502,73 @@ ECheckBoxState SCompositePassTree::GetGlobalEnabledState() const
 
 void SCompositePassTree::OnGlobalEnabledStateChanged(ECheckBoxState CheckBoxState)
 {
-	if (!Plate.IsValid())
+	if (!PassListOwner.IsValid())
 	{
 		return;
 	}
 
-	TSharedPtr<FScopedTransaction> AllEnabledTransaction;
+	TStrongObjectPtr<UObject> OwnerObject = PassListOwner->GetObject();
+	if (!OwnerObject.IsValid())
+	{
+		return;
+	}
 	
-	for (int32 Index = 0; Index < (int32)EPassType::MAX; ++Index)
+	TSharedPtr<FScopedTransaction> AllEnabledTransaction;
+
+	auto SetPassListEnabledStatus = [this, &AllEnabledTransaction, &OwnerObject, &CheckBoxState](int32 InGroupIndex)
 	{
-		const TArray<TObjectPtr<UCompositePassBase>>& PassList = GetPassList((EPassType)Index);
+		const TArray<TObjectPtr<UCompositePassBase>>& PassList = PassListOwner->GetPassesForGroup(InGroupIndex);
 
 		for (const TObjectPtr<UCompositePassBase>& Pass : PassList)
 		{
 			if (!AllEnabledTransaction.IsValid())
 			{
 				AllEnabledTransaction = MakeShared<FScopedTransaction>(LOCTEXT("AllEnabledTransaction", "Set All Enabled"));
-				Plate->Modify();
+				OwnerObject->Modify();
 			}
 
 			Pass->Modify();
 			Pass->SetEnabled(CheckBoxState == ECheckBoxState::Checked);
 		}
+	};
+	
+	const int32 NumGroups = PassListOwner->GetNumGroups();
+	if (NumGroups <= 0)
+	{
+		SetPassListEnabledStatus(INDEX_NONE);
+	}
+	else
+	{
+		for (int32 GroupIndex = 0; GroupIndex < NumGroups; ++GroupIndex)
+		{
+			SetPassListEnabledStatus(GroupIndex);
+		}
 	}
 }
 
-void SCompositePassTree::OnPassMoved(const TSharedPtr<FPassTreeItem>& InTreeItem, EPassType InDestPassType, int InDestIndex)
+void SCompositePassTree::OnPassMoved(const TSharedPtr<FPassTreeItem>& InTreeItem, int32 InDestGroupIndex, int InDestIndex)
 {
 	if (!InTreeItem.IsValid() ||
-		!InTreeItem->Plate.IsValid() ||
+		!InTreeItem->HasValidOwner() ||
 		!InTreeItem->HasValidPassIndex())
 	{
 		return;
 	}
 
 	FScopedTransaction MovePassTransaction(LOCTEXT("MovePassTransaction", "Move Pass"));
-	InTreeItem->Plate->Modify();
-	
-	TArray<TObjectPtr<UCompositePassBase>>& SourcePassList = GetPassList(InTreeItem->Plate, InTreeItem->PassType);
-	TArray<TObjectPtr<UCompositePassBase>>& DestPassList = GetPassList(InTreeItem->Plate, InDestPassType);
-
-	FProperty* SourcePassListProperty = GetPassListProperty(InTreeItem->PassType);
-	FProperty* DestPassListProperty = GetPassListProperty(InDestPassType);
-	
-	UCompositePassBase* Pass = SourcePassList[InTreeItem->PassIndex];
-	
-	InTreeItem->Plate->PreEditChange(SourcePassListProperty);
-	SourcePassList.RemoveAt(InTreeItem->PassIndex);
-
-	// If we aren't adding the item back to the same pass list, invoke a property changed event for its removal
-	if (InTreeItem->PassType != InDestPassType)
-	{
-		FPropertyChangedEvent SourceListChangedEvent(SourcePassListProperty, EPropertyChangeType::ArrayRemove);
-		InTreeItem->Plate->PostEditChangeProperty(SourceListChangedEvent);
-
-		InTreeItem->Plate->PreEditChange(DestPassListProperty);
-	}
+	PassListOwner->MovePass(InTreeItem->GroupIndex, InTreeItem->PassIndex, InDestGroupIndex, InDestIndex);
 	
-	DestPassList.Insert(Pass, InDestIndex);
-
-	// If the item moved from one list to another, invoke a property changed event for the add; otherwise, invoke a move event
-	if (InTreeItem->PassType != InDestPassType)
-	{
-		FPropertyChangedEvent DestListChangedEvent(DestPassListProperty, EPropertyChangeType::ArrayAdd);
-		InTreeItem->Plate->PostEditChangeProperty(DestListChangedEvent);
-	}
-	else
-	{
-		FPropertyChangedEvent MoveChangedEvent(SourcePassListProperty, EPropertyChangeType::ArrayMove);
-		InTreeItem->Plate->PostEditChangeProperty(MoveChangedEvent);
-	}
-	
-	RefreshPassType(InTreeItem->PassType, /* bRefreshTreeView */ false);
-	RefreshPassType(InDestPassType);
+	RefreshGroup(InTreeItem->GroupIndex, /* bRefreshTreeView */ false);
+	RefreshGroup(InDestGroupIndex);
 }
 
 void SCompositePassTree::OnPassSelectionChanged(TSharedPtr<FPassTreeItem> InTreeItem, ESelectInfo::Type SelectInfo)
 {
+	if (bSilenceSelectionChanges)
+	{
+		return;
+	}
+	
 	TArray<FPassTreeItemPtr> SelectedItems = TreeView.IsValid() ? TreeView->GetSelectedItems() : TArray<FPassTreeItemPtr>();
 	
 	TArray<UObject*> SelectedLayers;
@@ -1443,7 +1576,7 @@ void SCompositePassTree::OnPassSelectionChanged(TSharedPtr<FPassTreeItem> InTree
 
 	for (const FPassTreeItemPtr& Item : SelectedItems)
 	{
-		if (!Item->Plate.IsValid())
+		if (!Item->HasValidOwner())
 		{
 			continue;
 		}
@@ -1468,6 +1601,11 @@ void SCompositePassTree::OnFilterChanged()
 
 void SCompositePassTree::OnPassAdded(const UClass* InPassClass)
 {
+	if (!PassListOwner.IsValid() || !PassListOwner->IsObjectValid())
+	{
+		return;
+	}
+	
 	if (!IsValid(InPassClass) || InPassClass->IsChildOf<UCompositeLayerBase>())
 	{
 		return;
@@ -1476,61 +1614,84 @@ void SCompositePassTree::OnPassAdded(const UClass* InPassClass)
 	TArray<FPassTreeItemPtr> SelectedItems = TreeView.IsValid() ? TreeView->GetSelectedItems() : TArray<FPassTreeItemPtr>();
 	if (SelectedItems.IsEmpty())
 	{
-		// If there are no selected items, add the pass to the media passes list, unless it is a PPM pass, which can only get added to the layers passes list
-		if (InPassClass->IsChildOf<UCompositePassMaterial>())
+		// If there are no selected items, add the new pass to the default group, or after the last pass if there are no groups
+		const int32 DefaultGroupIndex = PassListOwner->GetDefaultGroupForNewPass(InPassClass);
+		if (DefaultGroupIndex != INDEX_NONE)
 		{
-			SelectedItems.Add(PassTreeItems[1]);
+			if (PassTreeItems.IsValidIndex(DefaultGroupIndex))
+			{
+				SelectedItems.Add(PassTreeItems[DefaultGroupIndex]);
+			}
+			else
+			{
+				SelectedItems.Add(PassTreeItems[0]);
+			}
 		}
 		else
 		{
-			SelectedItems.Add(PassTreeItems[0]);
+			if (!PassTreeItems.IsEmpty())
+			{
+				SelectedItems.Add(PassTreeItems.Last());
+			}
 		}
 	}
 
+	int32 GroupIndex = INDEX_NONE;
+	int32 PassIndex = PassTreeItems.Num();
+	
 	// Item after which the new pass will get added
-	FPassTreeItemPtr AnchorItem = SelectedItems.Last();
-	if (!AnchorItem.IsValid() || !AnchorItem->Plate.IsValid())
+	FPassTreeItemPtr AnchorItem = !SelectedItems.IsEmpty() ? SelectedItems.Last() : nullptr;
+	if (AnchorItem.IsValid())
 	{
-		return;
+		GroupIndex = AnchorItem->GroupIndex;
+		PassIndex = AnchorItem->PassIndex + 1;
 	}
 
-	// Post-process material pass can only be added to the Layer passes list
-	if (InPassClass->IsChildOf<UCompositePassMaterial>() && AnchorItem->PassType != EPassType::Layer)
+	if (!PassListOwner->CanAddPass(InPassClass, GroupIndex, PassIndex))
 	{
 		return;
 	}
 	
 	FScopedTransaction AddPassTransaction(LOCTEXT("AddPassTransaction", "Add Pass"));
-	AnchorItem->Plate->Modify();
+	const int32 NewPassIndex = PassListOwner->AddPass(InPassClass, GroupIndex, PassIndex);
 	
-	UCompositePassBase* NewPass = NewObject<UCompositePassBase>(AnchorItem->Plate.Get(), InPassClass, NAME_None, RF_Transactional);
-	TArray<TObjectPtr<UCompositePassBase>>& PassList = GetPassList(AnchorItem->Plate, AnchorItem->PassType);
-
-	FProperty* PassListProperty = GetPassListProperty(AnchorItem->PassType);
-	AnchorItem->Plate->PreEditChange(PassListProperty);
-	
-	int32 NewPassIndex = PassList.Insert(NewPass, AnchorItem->PassIndex + 1);
-
-	FPropertyChangedEvent PropertyChangedEvent(PassListProperty, EPropertyChangeType::ArrayAdd);
-	AnchorItem->Plate->PostEditChangeProperty(PropertyChangedEvent);
-
 	if (Toolbar.IsValid())
 	{
 		Toolbar->ClearFilters();
 	}
 	
-	RefreshPassType(AnchorItem->PassType);
+	RefreshGroup(GroupIndex);
 
-	FPassTreeItemPtr PassToSelect = PassTreeItems[(int32)AnchorItem->PassType]->Children[NewPassIndex];
+	FPassTreeItemPtr PassToSelect = nullptr;
+	if (GroupIndex != INDEX_NONE)
+	{
+		if (PassTreeItems.IsValidIndex(GroupIndex) && PassTreeItems[GroupIndex]->Children.IsValidIndex(NewPassIndex))
+		{
+			PassToSelect = PassTreeItems[GroupIndex]->Children[NewPassIndex];
+		}
+	}
+	else if (PassTreeItems.IsValidIndex(NewPassIndex) && PassTreeItems[NewPassIndex]->HasValidPassIndex())
+	{
+		PassToSelect = PassTreeItems[NewPassIndex];
+	}
+	
 	if (TreeView.IsValid())
 	{
 		TreeView->ClearSelection();
-		TreeView->SetItemSelection(PassToSelect, true);
+		if (PassToSelect.IsValid())
+		{
+			TreeView->SetItemSelection(PassToSelect, true);
+		}
 	}
 }
 
 bool SCompositePassTree::OnFilterNewPassType(const UClass* InPassType) const
 {
+	if (!PassListOwner.IsValid() || !PassListOwner->IsObjectValid())
+	{
+		return true;
+	}
+	
 	if (!IsValid(InPassType) || InPassType->IsChildOf<UCompositeLayerBase>())
 	{
 		return true;
@@ -1543,15 +1704,13 @@ bool SCompositePassTree::OnFilterNewPassType(const UClass* InPassType) const
 	}
 
 	TArray<FPassTreeItemPtr> SelectedItems = TreeView.IsValid() ? TreeView->GetSelectedItems() : TArray<FPassTreeItemPtr>();
-	const bool bLayerPassesSelected = SelectedItems.IsEmpty() || SelectedItems.ContainsByPredicate([](const FPassTreeItemPtr& InSelectedItem)
-	{
-		return InSelectedItem->PassType == EPassType::Layer;
-	});
-
-	// Post-process material pass can only be added to the Layer passes list
-	if (!bLayerPassesSelected && InPassType->IsChildOf<UCompositePassMaterial>())
+	for (const FPassTreeItemPtr& Item : SelectedItems)
 	{
-		return true;
+		if (!PassListOwner->CanAddPass(InPassType, Item->GroupIndex, Item->PassIndex))
+		{
+			// If the pass type can't be added at even one of the selected items, filter that type out
+			return true;
+		}
 	}
 	
 	return false;
@@ -1660,20 +1819,18 @@ FSlateColor SCompositePassTree::GetFilterStatusColor() const
 
 void SCompositePassTree::OnObjectPropertyChanged(UObject* InObject, FPropertyChangedEvent& InPropertyChangedEvent)
 {
-	if (!Plate.IsValid())
+	if (!PassListOwner.IsValid())
 	{
 		return;
 	}
 
-	if (InObject != Plate.Get())
+	if (InObject != PassListOwner->GetObject().Get())
 	{
 		return;
 	}
 
 	const FName PropName = InPropertyChangedEvent.GetPropertyName();
-	if (PropName == GET_MEMBER_NAME_CHECKED(UCompositeLayerPlate, MediaPasses) ||
-		PropName == GET_MEMBER_NAME_CHECKED(UCompositeLayerPlate, ScenePasses) ||
-		PropName == GET_MEMBER_NAME_CHECKED(UCompositeLayerPlate, LayerPasses))
+	if (PassListOwner->IsPassListPropertyName(PropName))
 	{
 		FillPassTreeItems();
 	}
@@ -1769,7 +1926,13 @@ public:
 
 void SCompositePassTree::PasteSelectedItems()
 {
-	if (!Plate.IsValid())
+	if (!PassListOwner.IsValid())
+	{
+		return;
+	}
+
+	TStrongObjectPtr<UObject> OwnerObject = PassListOwner->GetObject();
+	if (!OwnerObject.IsValid())
 	{
 		return;
 	}
@@ -1780,24 +1943,24 @@ void SCompositePassTree::PasteSelectedItems()
 		return;
 	}
 
-	// A list of pass types to paste the clipboard contents in, and the index to paste at
-	TMap<EPassType, int32> PassTypesToPasteTo;
+	// A list of groups to paste the clipboard contents in, and the index to paste at
+	TMap<int32, int32> GroupsToPasteTo;
 	for (const FPassTreeItemPtr& SelectedItem : SelectedItems)
 	{
 		if (SelectedItem.IsValid())
 		{
-			if (PassTypesToPasteTo.Contains(SelectedItem->PassType))
+			if (GroupsToPasteTo.Contains(SelectedItem->GroupIndex))
 			{
-				PassTypesToPasteTo[SelectedItem->PassType] = FMath::Max(SelectedItem->PassIndex + 1, PassTypesToPasteTo[SelectedItem->PassType]);
+				GroupsToPasteTo[SelectedItem->GroupIndex] = FMath::Max(SelectedItem->PassIndex + 1, GroupsToPasteTo[SelectedItem->GroupIndex]);
 			}
 			else
 			{
-				PassTypesToPasteTo.Add(SelectedItem->PassType, SelectedItem->PassIndex + 1);
+				GroupsToPasteTo.Add(SelectedItem->GroupIndex, SelectedItem->PassIndex + 1);
 			}
 		}
 	}
 
-	if (PassTypesToPasteTo.IsEmpty())
+	if (GroupsToPasteTo.IsEmpty())
 	{
 		return;
 	}
@@ -1813,35 +1976,18 @@ void SCompositePassTree::PasteSelectedItems()
 		return;
 	}
 
-	TMap<EPassType, TArray<int32>> PassesToSelect;
+	TMap<int32, TArray<int32>> PassesToSelect;
 	FScopedTransaction PastePassesTransaction(LOCTEXT("PastePassesTransaction", "Paste Passes"));
-	Plate->Modify();
-	for (const TPair<EPassType, int32>& Pair : PassTypesToPasteTo)
+	OwnerObject->Modify();
+	for (const TPair<int32, int32>& Pair : GroupsToPasteTo)
 	{
-		EPassType PassTypeToPasteAt = Pair.Key;
+		const int32 GroupToPasteAt = Pair.Key;
 		int32 IndexToPasteAt = Pair.Value;
 
-		TArray<TObjectPtr<UCompositePassBase>>& PassList = GetPassList(PassTypeToPasteAt);
-		
-		PassesToSelect.Add(PassTypeToPasteAt, TArray<int32>());
-		
-		FProperty* PassListProperty = GetPassListProperty(PassTypeToPasteAt);
-		Plate->PreEditChange(PassListProperty);
-		
-		for (UCompositePassBase* PassToPaste : Factory.CompositePasses)
-		{
-			UCompositePassBase* NewPass = DuplicateObject(PassToPaste, Plate.Get());
-			
-			PassList.Insert(NewPass, IndexToPasteAt);
-			PassesToSelect[PassTypeToPasteAt].Add(IndexToPasteAt);
-			++IndexToPasteAt;
-		}
-
-		FPropertyChangedEvent PropertyChangedEvent(PassListProperty, EPropertyChangeType::ArrayAdd);
-		Plate->PostEditChangeProperty(PropertyChangedEvent);
+		PassesToSelect.Add(GroupToPasteAt, PassListOwner->CopyPasses(Factory.CompositePasses, GroupToPasteAt, IndexToPasteAt));
 		
 		constexpr bool bRefreshTree = false;
-		RefreshPassType(PassTypeToPasteAt, bRefreshTree);
+		RefreshGroup(GroupToPasteAt, bRefreshTree);
 	}
 
 	FilterPassTreeItems();
@@ -1849,14 +1995,31 @@ void SCompositePassTree::PasteSelectedItems()
 
 	if (TreeView.IsValid())
 	{
-		for (const TPair<EPassType, TArray<int32>>& Pair : PassesToSelect)
+		for (const TPair<int32, TArray<int32>>& Pair : PassesToSelect)
 		{
-			EPassType PassType = Pair.Key;
-			for (int32 IndexToSelect : Pair.Value)
+			const int32 GroupIndex = Pair.Key;
+
+			if (GroupIndex != INDEX_NONE)
 			{
-				if (PassTreeItems[(int32)PassType]->Children.IsValidIndex(IndexToSelect))
+				if (PassTreeItems.IsValidIndex(GroupIndex))
 				{
-					TreeView->SetItemSelection(PassTreeItems[(int32)PassType]->Children[IndexToSelect], true);
+					for (int32 IndexToSelect : Pair.Value)
+					{
+						if (PassTreeItems[GroupIndex]->Children.IsValidIndex(IndexToSelect))
+						{
+							TreeView->SetItemSelection(PassTreeItems[GroupIndex]->Children[IndexToSelect], true);
+						}
+					}
+				}
+			}
+			else
+			{
+				for (int32 IndexToSelect : Pair.Value)
+				{
+					if (PassTreeItems.IsValidIndex(IndexToSelect) && PassTreeItems[IndexToSelect]->HasValidPassIndex())
+					{
+						TreeView->SetItemSelection(PassTreeItems[IndexToSelect], true);
+					}
 				}
 			}
 		}
@@ -1897,14 +2060,20 @@ bool SCompositePassTree::CanDuplicateSelectedItems()
 
 void SCompositePassTree::DeleteSelectedItems()
 {
-	if (!Plate.IsValid())
+	if (!PassListOwner.IsValid())
+	{
+		return;
+	}
+
+	TStrongObjectPtr<UObject> OwnerObject = PassListOwner->GetObject();
+	if (!OwnerObject.IsValid())
 	{
 		return;
 	}
 	
 	TArray<FPassTreeItemPtr> SelectedItems = TreeView.IsValid() ? TreeView->GetSelectedItems() : TArray<FPassTreeItemPtr>();
 
-	TMap<EPassType, TArray<int32>> PassesToDelete;
+	TMap<int32, TArray<int32>> PassesToDelete;
 	for (const FPassTreeItemPtr& Item : SelectedItems)
 	{
 		if (!Item.IsValid() || !Item->HasValidPassIndex())
@@ -1912,56 +2081,30 @@ void SCompositePassTree::DeleteSelectedItems()
 			continue;
 		}
 
-		if (PassesToDelete.Contains(Item->PassType))
+		if (PassesToDelete.Contains(Item->GroupIndex))
 		{
-			PassesToDelete[Item->PassType].Add(Item->PassIndex);
+			PassesToDelete[Item->GroupIndex].Add(Item->PassIndex);
 		}
 		else
 		{
-			PassesToDelete.Add(Item->PassType, TArray<int32> { Item->PassIndex });
+			PassesToDelete.Add(Item->GroupIndex, TArray<int32> { Item->PassIndex });
 		}
 	}
 
 	if (PassesToDelete.Num())
 	{
 		FScopedTransaction DeletePassTransaction(LOCTEXT("DeletePassTransaction", "Delete Pass"));
-		Plate->Modify();
+		OwnerObject->Modify();
 		
-		for (TPair<EPassType, TArray<int32>>& PassesToDeleteByType : PassesToDelete)
+		for (TPair<int32, TArray<int32>>& PassesToDeleteByType : PassesToDelete)
 		{
-			const EPassType PassType = PassesToDeleteByType.Key;
-			TArray<TObjectPtr<UCompositePassBase>>& PassList = GetPassList(PassType);
-			
-			// Delete the passes in reverse order so that indices don't get messed up as we delete
+			const int32 GroupIndex = PassesToDeleteByType.Key;
 			TArray<int32>& PassIndicesToDelete = PassesToDeleteByType.Value;
-			PassIndicesToDelete.Sort();
-
-			FProperty* PassListProperty = GetPassListProperty(PassType);
-			Plate->PreEditChange(PassListProperty);
-			
-			for (int32 Index = PassIndicesToDelete.Num() - 1; Index >= 0; --Index)
-			{
-				const int32 PassIndex = PassIndicesToDelete[Index];
-
-				// Here we replicate details panels instanced property array removal code.
-				// The removed instanced pass objects must be moved to the transient package.
-				TObjectPtr<UCompositePassBase>& PassToRemove = PassList[PassIndex];
-				if (IsValid(PassToRemove))
-				{
-					PassToRemove->Modify();
-					PassToRemove->Rename(nullptr, GetTransientPackage(), REN_DontCreateRedirectors);
-				}
-
-				PassList.RemoveAt(PassIndex, EAllowShrinking::No);
-			}
 
-			PassList.Shrink();
-			
-			FPropertyChangedEvent PropertyChangedEvent(PassListProperty, EPropertyChangeType::ArrayAdd);
-			Plate->PostEditChangeProperty(PropertyChangedEvent);
+			PassListOwner->RemovePasses(GroupIndex, PassIndicesToDelete);
 			
 			constexpr bool bRefreshTreeView = false;
-			RefreshPassType(PassType, bRefreshTreeView);
+			RefreshGroup(GroupIndex, bRefreshTreeView);
 		}
 
 		RefreshAndExpandTreeView();
@@ -2007,6 +2150,17 @@ bool SCompositePassTree::CanRenameSelectedItem()
 
 void SCompositePassTree::EnableSelectedItems()
 {
+	if (!PassListOwner.IsValid())
+	{
+		return;
+	}
+
+	TStrongObjectPtr<UObject> OwnerObject = PassListOwner->GetObject();
+	if (!OwnerObject.IsValid())
+	{
+		return;
+	}
+	
 	TArray<FPassTreeItemPtr> SelectedItems = TreeView.IsValid() ? TreeView->GetSelectedItems() : TArray<FPassTreeItemPtr>();
 
 	bool bAnyEnabled = false;
@@ -2043,7 +2197,7 @@ void SCompositePassTree::EnableSelectedItems()
 	const bool bSetEnabled = !bAnyEnabled;
 
 	FScopedTransaction SetEnabledTransaction(LOCTEXT("SetEnabledTransaction", "Set Enabled"));
-	Plate->Modify();
+	OwnerObject->Modify();
 
 	for (const FPassTreeItemPtr& Item : SelectedItems)
 	{
diff --git a/Engine/Plugins/Compositing/Composite/Source/CompositeEditor/Private/UI/SCompositePassTree.h b/Engine/Plugins/Compositing/Composite/Source/CompositeEditor/Private/UI/SCompositePassTree.h
--- a/Engine/Plugins/Compositing/Composite/Source/CompositeEditor/Private/UI/SCompositePassTree.h
+++ b/Engine/Plugins/Compositing/Composite/Source/CompositeEditor/Private/UI/SCompositePassTree.h
@@ -2,55 +2,117 @@
 
 #pragma once
 
+#include "EditorUndoClient.h"
 #include "Widgets/SCompoundWidget.h"
 
+struct FSlateBrush;
 class FUICommandList;
 class SCompositePassTreeToolbar;
 class UCompositePassBase;
 class UCompositeLayerPlate;
 template<typename T> class STreeView;
 
+/** Wrapper around a UObject that owns one or more lists of composite passes to display in the pass tree */
+class ICompositePassListOwner
+{
+public:
+	/** Gets whether the backing UObject that owns the pass list(s) is valid */
+	virtual bool IsObjectValid() = 0;
+
+	/** Gets a strong pointer to the backing UObject that owns the pass list(s) */
+	virtual TStrongObjectPtr<UObject> GetObject() = 0;
+
+	/** Gets whether the specified property name matches any of the properties related to the pass list(s) being provided */
+	virtual bool IsPassListPropertyName(const FName& InPropertyName) = 0;
+
+	/** Gets whether the specified group and pass indices reference a valid pass in the owner's pass list(s) */
+	virtual bool IsValidPassIndex(int32 InGroupIndex, int32 InPassIndex) = 0;
+
+	/** Gets the pass for the specified group and pass indices, or nullptr if no pass is found */
+	virtual UCompositePassBase* GetPass(int32 InGroupIndex, int32 InPassIndex) = 0;
+
+	/** Gets the number of different groups all the passes managed by the owner should be stored and displayed in */
+	virtual int32 GetNumGroups() const = 0;
+
+	/** Gets all the passes for the specified group index; use INDEX_NONE if there are no distinct groups */
+	virtual TArray<TObjectPtr<UCompositePassBase>>& GetPassesForGroup(int32 InGroupIndex) = 0;
+
+	/** Gets the string used for filtering in the tree view for a specified pass group */
+	virtual FString GetGroupFilterString(int32 InGroupIndex) = 0;
+
+	/** Gets the icon used for displaying the group in the tree view for a specified pass group */
+	virtual const FSlateBrush* GetGroupIcon(int32 InGroupIndex) = 0;
+
+	/** Gets the display name used for displaying the group in the tree view for a specified pass group */
+	virtual FText GetGroupDisplayName(int32 InGroupIndex) = 0;
+
+	struct FGroupFilterConfig { FString FilterName; FText DisplayName; FText ToolTip; };
+	/** Gets the filter config used to display filtering pills in the tree view for a specified pass group */
+	virtual FGroupFilterConfig GetGroupFilterConfig(int32 InGroupIndex) = 0;
+
+	/** Gets the default group a new pass should be added in when no group is selected, or INDEX_NONE if the pass should not be added to a group */
+	virtual int32 GetDefaultGroupForNewPass(const UClass* InPassClass) const = 0;
+
+	/** Gets whether a pass of the specified type can be added to the specified group at the specified location */
+	virtual bool CanAddPass(const UClass* InPassClass, int32 InGroupIndex, int32 InPassIndex) const = 0;
+
+	/** Adds a pass of the specified type to the specified group at the specified location */
+	virtual int32 AddPass(const UClass* InPassClass, int32 InGroupIndex, int32 InPassIndex) = 0;
+
+	/** Copies all the specified passes into the specified group at the specified location */
+	virtual TArray<int32> CopyPasses(const TArray<UCompositePassBase*>& InPassesToCopy, int32 InGroupIndex, int32 InPassIndex) = 0;
+
+	/** Moves a pass from the specified group and location to the specified destination group and location */
+	virtual void MovePass(int32 InSourceGroupIndex, int32 InSourcePassIndex, int32 InDestGroupIndex, int32 InDestPassIndex) = 0;
+
+	/** Removes all passes at the specified group */
+	virtual void RemovePasses(int32 InGroupIndex, const TArray<int32>& InPassIndices) = 0;
+};
+
 /**
  * Widget that displays a filterable tree view for the passes within a composite plate layer
  */
-class SCompositePassTree : public SCompoundWidget
+class SCompositePassTree : public SCompoundWidget, public FSelfRegisteringEditorUndoClient
 {
 public:
 	DECLARE_DELEGATE_OneParam(FOnSelectionChanged, const TArray<UObject*>&);
 	
 private:
-	enum class EPassType : uint8
-	{
-		Media,
-		//Scene, // TODO: Add Scene when it is supported
-		Layer,
-		MAX
-	};
-
 	/** Stores data needed for a single row item in the tree view */
 	struct FPassTreeItem
 	{
-		/** Plate whose passes are being displayed */
-		TWeakObjectPtr<UCompositeLayerPlate> Plate = nullptr;
-
-		/** The pass type this tree item belongs to */
-		EPassType PassType = EPassType::Media;
+		/** Owner object whose passes are being displayed */
+		TSharedPtr<ICompositePassListOwner> Owner = nullptr;
 
 		/** If this tree item is for a specific pass, this is the pass's index in the plate's corresponding list of passes */
 		int32 PassIndex = INDEX_NONE;
 
+		/** The group this tree item belongs to */
+		int32 GroupIndex = INDEX_NONE;
+		
 		/** Indicates if this tree item is being filtered out by the current filter */
 		bool bFilteredOut = false;
 
 		/** Child tree items of this tree item */
 		TArray<TSharedPtr<FPassTreeItem>> Children;
 		
+		/** Gets whether the pass list owner pointer is valid */
+		bool HasValidOwner() const;
 		
 		/** Gets whether the pass index is valid within the plate's list of passes */
 		bool HasValidPassIndex() const;
 
 		/** Gets the pass object referenced by this tree item */
 		UCompositePassBase* GetPass() const;
+
+		/** If this item is a group item, gets the filter string to use for the item */
+		FString GetGroupFilterString() const;
+
+		/** If this item is a group item, gets the icon to use for the item */
+		const FSlateBrush* GetGroupIcon() const;
+
+		/** If this item is a group item, gets the display name to use for the item */
+		FText GetGroupDisplayName() const;
 	};
 	using FPassTreeItemPtr = TSharedPtr<FPassTreeItem>;
 	
@@ -60,7 +122,7 @@ public:
 		SLATE_EVENT(FSimpleDelegate, OnLayoutChanged)
 	SLATE_END_ARGS()
 
-	void Construct(const FArguments& InArgs, const TWeakObjectPtr<UCompositeLayerPlate>& InPlate);
+	void Construct(const FArguments& InArgs, const TSharedPtr<ICompositePassListOwner>& InPassListOwner);
 
 	virtual ~SCompositePassTree() override;
 	
@@ -68,6 +130,11 @@ public:
 	virtual FReply OnKeyDown(const FGeometry& MyGeometry, const FKeyEvent& InKeyEvent) override;
 	// End of SWidget interface
 
+	//~ Begin FSelfRegisteringEditorUndoClient interface
+	virtual void PostUndo(bool bSuccess) override;
+	virtual void PostRedo(bool bSuccess) override;
+	//~ End FSelfRegisteringEditorUndoClient interface
+	
 	/** Selects the specified passes in the pass tree */
 	void SelectPasses(const TArray<UCompositePassBase*>& InPasses);
 
@@ -75,25 +142,17 @@ public:
 	TArray<UCompositePassBase*> GetSelectedPasses() const;
 	
 private:
-	/** Gets a reference to the pass list that corresponds to the specified pass type in the plate. The plate pointer must be valid to call properly */
-	static TArray<TObjectPtr<UCompositePassBase>>& GetPassList(TWeakObjectPtr<UCompositeLayerPlate> InPlate, EPassType InPassType);
-	TArray<TObjectPtr<UCompositePassBase>>& GetPassList(EPassType InPassType) { return GetPassList(Plate, InPassType); }
-	const TArray<TObjectPtr<UCompositePassBase>>& GetPassList(EPassType InPassType) const { return GetPassList(Plate, InPassType); }
-	
-	/** Gets the FProperty for the specified pass type */
-	static FProperty* GetPassListProperty(EPassType InPassType);
-
 	/** Binds commands in the command list to callbacks */
 	void BindCommands();
 	
 	/** Creates any tree items needed for all passes in the plate */
-	void FillPassTreeItems();
+	void FillPassTreeItems(bool bPreserveSelection = false);
 
 	/** Applies the current filter to the pass tree items to build a filtered list of tree items */
 	void FilterPassTreeItems();
 
-	/** Refreshes the tree items for the passes in the specified pass type list */
-	void RefreshPassType(EPassType InPassType, bool bRefreshTreeView = true);
+	/** Refreshes the tree items for the passes in the specified group */
+	void RefreshGroup(int32 InGroupIndex, bool bRefreshTreeView = true);
 	
 	/** Refreshes the tree view's displayed list and automatically expands all tree items */
 	void RefreshAndExpandTreeView();
@@ -108,7 +167,7 @@ private:
 	void OnGlobalEnabledStateChanged(ECheckBoxState CheckBoxState);
 	
 	/** Raised when a pass has been moved via a drag drop operation */
-	void OnPassMoved(const TSharedPtr<FPassTreeItem>& InTreeItem, EPassType InDestPassType, int InDestIndex);
+	void OnPassMoved(const TSharedPtr<FPassTreeItem>& InTreeItem, int32 InDestGroupIndex, int InDestIndex);
 
 	/** Raised when the selected items in the tree view have bene changed */
 	void OnPassSelectionChanged(TSharedPtr<FPassTreeItem> InTreeItem, ESelectInfo::Type SelectInfo);
@@ -148,8 +207,8 @@ private:
 	bool CanEnableSelectedItems();
 	
 private:
-	/** The plate whose passes are being displayed */
-	TWeakObjectPtr<UCompositeLayerPlate> Plate;
+	/** The owner whose passes are being displayed */
+	TSharedPtr<ICompositePassListOwner> PassListOwner;
 
 	/** List of all pass tree items for the plate */
 	TArray<FPassTreeItemPtr> PassTreeItems;
@@ -165,6 +224,9 @@ private:
 
 	/** The command list for commands related to the tree view */
 	TSharedPtr<FUICommandList> CommandList;
+
+	/** When true, selection changes in the tree view will not invoke the SelectionChanged delegate */
+	bool bSilenceSelectionChanges;
 	
 	/** Delegate that gets raised when the tree view selection changes */
 	FOnSelectionChanged OnSelectionChanged;
diff --git a/Engine/Plugins/Compositing/Composite/Source/CompositeEditor/Private/UI/SCompositePassTreePanel.cpp b/Engine/Plugins/Compositing/Composite/Source/CompositeEditor/Private/UI/SCompositePassTreePanel.cpp
new file mode 100644
--- /dev/null
+++ b/Engine/Plugins/Compositing/Composite/Source/CompositeEditor/Private/UI/SCompositePassTreePanel.cpp
@@ -0,0 +1,111 @@
+﻿// Copyright Epic Games, Inc. All Rights Reserved.
+
+
+#include "SCompositePassTreePanel.h"
+
+#include "DetailsViewArgs.h"
+#include "PropertyEditorModule.h"
+#include "PropertyPath.h"
+#include "SCompositeEditorPanel.h"
+#include "SCompositePassTree.h"
+#include "Customizations/CompositePassBaseCustomization.h"
+#include "Layers/CompositeLayerPlate.h"
+#include "Modules/ModuleManager.h"
+
+void SCompositePassTreePanel::Construct(const FArguments& InArgs, const TSharedPtr<ICompositePassListOwner>& InPassListOwner)
+{
+	OnLayoutSizeChanged = InArgs._OnLayoutSizeChanged;
+	
+	FDetailsViewArgs DetailsViewArgs;
+	DetailsViewArgs.NameAreaSettings = FDetailsViewArgs::HideNameArea;
+	DetailsViewArgs.bAllowMultipleTopLevelObjects = true;
+	DetailsViewArgs.bShowPropertyMatrixButton = false;
+	DetailsViewArgs.bShowScrollBar = false;
+	DetailsViewArgs.bAllowSearch = false;
+	
+	DetailsView = FModuleManager::GetModuleChecked<FPropertyEditorModule>("PropertyEditor").CreateDetailView(DetailsViewArgs);
+
+	DetailsView->RegisterInstancedCustomPropertyLayout(
+		UCompositePassBase::StaticClass(),
+		FOnGetDetailCustomizationInstance::CreateStatic(&FCompositePassBaseCustomization::MakeInstance));
+
+	ChildSlot
+	[
+		SNew(SVerticalBox)
+
+		+SVerticalBox::Slot()
+		.AutoHeight()
+		[
+			SAssignNew(TreeView, SCompositePassTree, InPassListOwner)
+			.OnSelectionChanged(this, &SCompositePassTreePanel::OnTreeViewSelectionChanged)
+			.OnLayoutChanged_Lambda([this]
+			{
+				OnLayoutSizeChanged.ExecuteIfBound();
+			})
+		]
+
+		+SVerticalBox::Slot()
+		.AutoHeight()
+		[
+			DetailsView.ToSharedRef()
+		]
+	];
+}
+
+void SCompositePassTreePanel::Tick(const FGeometry& AllottedGeometry, const double InCurrentTime, const float InDeltaTime)
+{
+	SCompoundWidget::Tick(AllottedGeometry, InCurrentTime, InDeltaTime);
+
+	// This panel often changes its size dynamically due to the details panel (showing new objects, expanding properties, etc.)
+	// When this panel is embedded into a parent details panel (or any list/tree view), the dynamic size is an issue when it comes to
+	// determining scroll bar visibility and size. As such, this panel invokes a OnLayoutSizeChange event that allows parent containers
+	// to request a layout refresh, where the correct scroll bar size can be computed. To detect if the details view may have changed size,
+	// we can use GetPropertyRowNumbers, which returns a list of properties that are actively being displayed as widgets (thus will exclude collapsed
+	// properties). Caching this count on tick allows this panel to detect if the object properties being displayed have changed or been expanded.
+	// The actual layout size changed event must be invoked one tick after a row count difference has been detected to give the details view
+	// time to actually change its rendered geometry
+	
+	if (bLayoutSizeChanged)
+	{
+		OnLayoutSizeChanged.ExecuteIfBound();
+		bLayoutSizeChanged = false;
+	}
+	
+	const int32 RowCount = DetailsView->GetPropertyRowNumbers().Num();	
+ 	if (CachedDetailsViewRowCount != RowCount)
+	{
+		CachedDetailsViewRowCount = RowCount;
+ 		bLayoutSizeChanged = true;
+	}
+}
+
+void SCompositePassTreePanel::SelectPasses(TArray<UCompositePassBase*>& InPasses)
+{
+	if (TreeView.IsValid())
+	{
+		TreeView->SelectPasses(InPasses);
+	}
+}
+
+TArray<UCompositePassBase*> SCompositePassTreePanel::GetSelectedPasses() const
+{
+	if (TreeView)
+	{
+		return TreeView->GetSelectedPasses();
+	}
+
+	return { };
+}
+
+void SCompositePassTreePanel::OnTreeViewSelectionChanged(const TArray<UObject*>& Objects)
+{
+	if (!DetailsView.IsValid())
+	{
+		return;
+	}
+
+	DetailsView->SetObjects(Objects);
+
+	SCompositeEditorPanel::GetOnSelectionChanged().Broadcast();
+}
+
diff --git a/Engine/Plugins/Compositing/Composite/Source/CompositeEditor/Private/UI/SCompositePassTreePanel.h b/Engine/Plugins/Compositing/Composite/Source/CompositeEditor/Private/UI/SCompositePassTreePanel.h
new file mode 100644
--- /dev/null
+++ b/Engine/Plugins/Compositing/Composite/Source/CompositeEditor/Private/UI/SCompositePassTreePanel.h
@@ -0,0 +1,54 @@
+﻿// Copyright Epic Games, Inc. All Rights Reserved.
+
+#pragma once
+
+#include "Widgets/SCompoundWidget.h"
+
+class IDetailsView;
+class ICompositePassListOwner;
+class SCompositePassTree;
+class UCompositePassBase;
+class UCompositeLayerPlate;
+
+/**
+ * A panel that displays a tree view and a details panel for passes within the plate composite layer
+ */
+class SCompositePassTreePanel : public SCompoundWidget
+{
+public:
+	SLATE_BEGIN_ARGS(SCompositePassTreePanel) { }
+		SLATE_EVENT(FSimpleDelegate, OnLayoutSizeChanged)
+	SLATE_END_ARGS()
+
+	void Construct(const FArguments& InArgs, const TSharedPtr<ICompositePassListOwner>& InPassListOwner);
+
+	// SWidget interface
+	virtual void Tick(const FGeometry& AllottedGeometry, const double InCurrentTime, const float InDeltaTime) override;
+	// End of SWidget interface
+
+	/** Selects the specified passes in the pass tree */
+	void SelectPasses(TArray<UCompositePassBase*>& InPasses);
+
+	/** Gets a list of currently selected passes in the pass tree */
+	TArray<UCompositePassBase*> GetSelectedPasses() const;
+	
+private:
+	/** Raised when the selected items in the tree view have changed */
+	void OnTreeViewSelectionChanged(const TArray<UObject*>& Objects);
+	
+private:
+	/** Tree view that displays a hierarchical list of passes in the plate */
+	TSharedPtr<SCompositePassTree> TreeView;
+
+	/** The details view to display selected pass properties */
+	TSharedPtr<IDetailsView> DetailsView;
+
+	/** To identify when the details view has changed size by showing a new object or expanding a property group, keep track of its reported row count every tick */
+	int32 CachedDetailsViewRowCount = 0;
+
+	/** Flag used to indicate if OnLayoutSizeChanged needs to be called on the next tick */
+	bool bLayoutSizeChanged = false;
+	
+	/** Delegate that is raised when a potential size change has been detected in the panel (e.g. the details panel is displaying a new object) */
+	FSimpleDelegate OnLayoutSizeChanged;
+};
diff --git a/Engine/Plugins/Compositing/Composite/Source/CompositeEditor/Private/UI/SCompositePlatePassPanel.cpp b/Engine/Plugins/Compositing/Composite/Source/CompositeEditor/Private/UI/SCompositePlatePassPanel.cpp
deleted file mode 100644
diff --git a/Engine/Plugins/Compositing/Composite/Source/CompositeEditor/Private/UI/SCompositePlatePassPanel.h b/Engine/Plugins/Compositing/Composite/Source/CompositeEditor/Private/UI/SCompositePlatePassPanel.h
deleted file mode 100644
