Composure: Quick fix to planar reflections, to prevent compositing passes from being executed inside the scene capture render.

#Synced-CL 51392841

diff --git a/Engine/Plugins/Compositing/Composite/Source/Composite/Private/Layers/CompositeLayerPlanarReflection.cpp b/Engine/Plugins/Compositing/Composite/Source/Composite/Private/Layers/CompositeLayerPlanarReflection.cpp
--- a/Engine/Plugins/Compositing/Composite/Source/Composite/Private/Layers/CompositeLayerPlanarReflection.cpp
+++ b/Engine/Plugins/Compositing/Composite/Source/Composite/Private/Layers/CompositeLayerPlanarReflection.cpp
@@ -368,7 +368,12 @@ UCompositeSceneCapture2DComponent* UCompositeLayerPlanarReflection::GetSceneCapt
 			CaptureComponent->CaptureSource = UE::Composite::Private::GetPlanarReflectionCaptureSource(*CompositeActor, bCustomRenderPass);
 			CaptureComponent->bRenderInMainRenderer = bCustomRenderPass;
 			CaptureComponent->bMainViewCamera = false;
-			CaptureComponent->ShowFlags.SetAllowPrimitiveAlphaHoldout(true);
+
+			/**
+			* Fixme: AllowPrimitiveAlphaHoldout is still used to prevent compositing while rendering
+			* scene captures in FCompositeCoreSceneViewExtension, so we cannot allow it currently.
+			*/
+			//CaptureComponent->ShowFlags.SetAllowPrimitiveAlphaHoldout(true);
 		}
 
 		return CaptureComponent;
