Composure: Minor shadow/reflection layer fix to account for mesh visibility.

#Synced-CL 51392841

diff --git a/Engine/Plugins/Compositing/Composite/Source/Composite/Private/Layers/CompositeLayerShadowReflection.cpp b/Engine/Plugins/Compositing/Composite/Source/Composite/Private/Layers/CompositeLayerShadowReflection.cpp
--- a/Engine/Plugins/Compositing/Composite/Source/Composite/Private/Layers/CompositeLayerShadowReflection.cpp
+++ b/Engine/Plugins/Compositing/Composite/Source/Composite/Private/Layers/CompositeLayerShadowReflection.cpp
@@ -130,6 +130,8 @@ void UCompositeLayerShadowReflection::PostTransacted(const FTransactionObjectEve
 
 void UCompositeLayerShadowReflection::OnEndOfFrameUpdate(UWorld* InWorld)
 {
+	// Note: OnEndOfFrameUpdate is only called from the base layer if IsEnabled() is true.
+
 	/**
 	* In order to retain the ability to reference actors in sublevels, we prefer using weak pointer components on scene captures. However,
 	* primitive components may often be recreated (see RerunConstructionScripts(), FStaticMeshComponentRecreateRenderStateContext, dynamic
@@ -139,6 +141,18 @@ void UCompositeLayerShadowReflection::OnEndOfFrameUpdate(UWorld* InWorld)
 	{
 		UpdateSceneCaptureComponents(*CachedSceneCaptures[1], Actors);
 	}
+
+	/** 
+	 * We also re-evaluate visibility-dependent flags each frame so that runtime
+	 * visibility changes are present in the shadow/reflection catcher.
+	 */
+	if (!Actors.IsEmpty())
+	{
+		constexpr bool bSceneCaptureHidden = true;
+		constexpr bool bConditionOnVisibility = true;
+
+		UpdatePrimitiveVisibilityState(bSceneCaptureHidden, bSceneCaptureHidden, bSceneCaptureHidden, bConditionOnVisibility);
+	}
 }
 
 bool UCompositeLayerShadowReflection::GetProxy(const UE::CompositeCore::FPassInputDecl& /*InputDecl*/, FCompositeTraversalContext& InContext, FSceneRenderingBulkObjectAllocator& InFrameAllocator, FCompositeCorePassProxy*& OutProxy) const
@@ -253,41 +267,59 @@ void UCompositeLayerShadowReflection::SetEnabled(bool bInEnabled)
 void UCompositeLayerShadowReflection::UpdatePrimitiveVisibilityState(
 	bool bInHiddenInSceneCapture,
 	bool bInAffectIndirectLightingWhileHidden,
-	bool bInCastHiddenShadow) const
+	bool bInCastHiddenShadow,
+	bool bConditionOnVisibility) const
 {
 	for (const TSoftObjectPtr<AActor>& Actor : Actors)
 	{
-		if (Actor.IsValid())
+		if (!Actor.IsValid())
+		{
+			continue;
+		}
+
+		for (UActorComponent* Component : Actor->GetComponents())
 		{
-			for (UActorComponent* Component : Actor->GetComponents())
+			UPrimitiveComponent* PrimitiveComponent = Cast<UPrimitiveComponent>(Component);
+
+			if (!PrimitiveComponent)
 			{
-				if (UPrimitiveComponent* PrimitiveComponent = Cast<UPrimitiveComponent>(Component))
-				{
-					bool bUpdatedComponent = false;
-
-					if (PrimitiveComponent->bHiddenInSceneCapture != bInHiddenInSceneCapture)
-					{
-						PrimitiveComponent->SetHiddenInSceneCapture(bInHiddenInSceneCapture);
-						bUpdatedComponent = true;
-					}
-
-					if (PrimitiveComponent->bAffectIndirectLightingWhileHidden != bInAffectIndirectLightingWhileHidden)
-					{
-						PrimitiveComponent->SetAffectIndirectLightingWhileHidden(bInAffectIndirectLightingWhileHidden);
-						bUpdatedComponent = true;
-					}
-
-					if (PrimitiveComponent->bCastHiddenShadow != bInCastHiddenShadow)
-					{
-						PrimitiveComponent->SetCastHiddenShadow(bInCastHiddenShadow);
-						bUpdatedComponent = true;
-					}
-
-					if (bUpdatedComponent)
-					{
-						PrimitiveComponent->Modify();
-					}
-				}
+				continue;
+			}
+
+			bool bUpdatedComponent = false;
+			bool bEffectiveIndirectLighting = bInAffectIndirectLightingWhileHidden;
+			bool bEffectiveCastHiddenShadow = bInCastHiddenShadow;
+
+			// Respect primitive visibility: invisible meshes should not cast hidden shadows or affect indirect lighting.
+			if (bConditionOnVisibility)
+			{
+				const bool bShouldRender = PrimitiveComponent->ShouldRender();
+
+				bEffectiveIndirectLighting &= bShouldRender;
+				bEffectiveCastHiddenShadow &= bShouldRender;
+			}
+
+			if (PrimitiveComponent->bHiddenInSceneCapture != bInHiddenInSceneCapture)
+			{
+				PrimitiveComponent->SetHiddenInSceneCapture(bInHiddenInSceneCapture);
+				bUpdatedComponent = true;
+			}
+
+			if (PrimitiveComponent->bAffectIndirectLightingWhileHidden != bEffectiveIndirectLighting)
+			{
+				PrimitiveComponent->SetAffectIndirectLightingWhileHidden(bEffectiveIndirectLighting);
+				bUpdatedComponent = true;
+			}
+
+			if (PrimitiveComponent->bCastHiddenShadow != bEffectiveCastHiddenShadow)
+			{
+				PrimitiveComponent->SetCastHiddenShadow(bEffectiveCastHiddenShadow);
+				bUpdatedComponent = true;
+			}
+
+			if (bUpdatedComponent)
+			{
+				PrimitiveComponent->Modify();
 			}
 		}
 	}
@@ -303,10 +335,10 @@ void UCompositeLayerShadowReflection::SetHiddenPrimitiveState(bool bInHiddenInSc
 
 	if (bInHiddenInSceneCapture)
 	{
-		constexpr bool bHiddenInSceneCapture = true;
-		constexpr bool bAffectIndirectLightingWhileHidden = true;
-		constexpr bool bCastHiddenShadow = true;
-		UpdatePrimitiveVisibilityState(bHiddenInSceneCapture, bAffectIndirectLightingWhileHidden, bCastHiddenShadow);
+		constexpr bool bSceneCaptureHidden = true;
+		constexpr bool bConditionOnVisibility = true;
+
+		UpdatePrimitiveVisibilityState(bSceneCaptureHidden, bSceneCaptureHidden, bSceneCaptureHidden, bConditionOnVisibility);
 
 		if (ActiveHiddenStateLayers.IsEmpty())
 		{
@@ -320,8 +352,9 @@ void UCompositeLayerShadowReflection::SetHiddenPrimitiveState(bool bInHiddenInSc
 	{
 		// Restore to class default
 		UPrimitiveComponent* CDO = UPrimitiveComponent::StaticClass()->GetDefaultObject<UPrimitiveComponent>();
+		constexpr bool bConditionOnVisibility = false;
 
-		UpdatePrimitiveVisibilityState(CDO->bHiddenInSceneCapture, CDO->bAffectIndirectLightingWhileHidden, CDO->bCastHiddenShadow);
+		UpdatePrimitiveVisibilityState(CDO->bHiddenInSceneCapture, CDO->bAffectIndirectLightingWhileHidden, CDO->bCastHiddenShadow, bConditionOnVisibility);
 
 		const int32 NumRemoved = ActiveHiddenStateLayers.Remove(this);
 
diff --git a/Engine/Plugins/Compositing/Composite/Source/Composite/Public/Layers/CompositeLayerShadowReflection.h b/Engine/Plugins/Compositing/Composite/Source/Composite/Public/Layers/CompositeLayerShadowReflection.h
--- a/Engine/Plugins/Compositing/Composite/Source/Composite/Public/Layers/CompositeLayerShadowReflection.h
+++ b/Engine/Plugins/Compositing/Composite/Source/Composite/Public/Layers/CompositeLayerShadowReflection.h
@@ -110,7 +110,7 @@ PRAGMA_ENABLE_DEPRECATION_WARNINGS
 
 private:
 	/** Update scene capture visibility on registered mesh actors primitives. */
-	void UpdatePrimitiveVisibilityState(bool bInHiddenInSceneCapture, bool bInAffectIndirectLightingWhileHidden, bool bInCastHiddenShadow) const;
+	void UpdatePrimitiveVisibilityState(bool bInHiddenInSceneCapture, bool bInAffectIndirectLightingWhileHidden, bool bInCastHiddenShadow, bool bConditionOnVisibility) const;
 
 	/** Set primitive visibility state and console variables required for hidden in scene capture. */
 	void SetHiddenPrimitiveState(bool bInHiddenInSceneCapture) const;
