Composure: Add pass support on main render layer.

Since the main render scene color will often use a sub-rectangle of the full texture (to prevent resizing reallocation of scene textures), we need to handle view rectangles correctly. We therefore refactor all passes to use proper input/output view rectangle definitions, and FPixelShaderUtils::AddFullscreenPass where possible.

#Synced-CL 51392841

diff --git a/Engine/Plugins/Compositing/Composite/Shaders/Private/CompositeAlphaTransform.usf b/Engine/Plugins/Compositing/Composite/Shaders/Private/CompositeAlphaTransform.usf
--- a/Engine/Plugins/Compositing/Composite/Shaders/Private/CompositeAlphaTransform.usf
+++ b/Engine/Plugins/Compositing/Composite/Shaders/Private/CompositeAlphaTransform.usf
@@ -1,11 +1,14 @@
 // Copyright Epic Games, Inc. All Rights Reserved.
 
 #include "/Engine/Private/Common.ush"
+#include "/Engine/Private/ScreenPass.ush"
 
 #include "CompositeCommon.ush"
 
-Texture2D Texture;
-SamplerState Sampler;
+FScreenTransform SvPositionToInputTextureUV;
+Texture2D Input_Texture;
+SamplerState Input_Sampler;
+
 uint PremultOp;
 uint bInvertSourceAlpha;
 uint OverrideOutputAlpha;
@@ -19,11 +22,12 @@ uint OverrideOutputAlpha;
 #define OverrideOutput_One 2u
 
 void MainPS(
-    noperspective float4 InUVAndScreenPos : TEXCOORD0,
+    float4 SvPosition : SV_POSITION,
     out float4 OutColor : SV_Target0
     )
 {
-    float4 Color = Texture.SampleLevel(Sampler, InUVAndScreenPos.xy, 0);
+    float2 TextureUV = ApplyScreenTransform(SvPosition.xy, SvPositionToInputTextureUV);
+    float4 Color = Input_Texture.SampleLevel(Input_Sampler, TextureUV, 0);
 
 	Color.a = select(bInvertSourceAlpha > 0, saturate(1.0 - Color.a), Color.a);
 	
diff --git a/Engine/Plugins/Compositing/Composite/Shaders/Private/CompositeBlur.usf b/Engine/Plugins/Compositing/Composite/Shaders/Private/CompositeBlur.usf
--- a/Engine/Plugins/Compositing/Composite/Shaders/Private/CompositeBlur.usf
+++ b/Engine/Plugins/Compositing/Composite/Shaders/Private/CompositeBlur.usf
@@ -1,14 +1,16 @@
 // Copyright Epic Games, Inc. All Rights Reserved.
 #include "/Engine/Private/Common.ush"
+#include "/Engine/Private/ScreenPass.ush"
 #include "/Engine/Public/Platform.ush"
 
 // See https://www.rastergrid.com/blog/2010/09/efficient-gaussian-blur-with-linear-sampling/
 
-Texture2D InputTexture;
-SamplerState InputSampler;
+Texture2D Input_Texture;
+SamplerState Input_Sampler;
+FScreenTransform SvPositionToInputTextureUV;
+SCREEN_PASS_TEXTURE_VIEWPORT(Input)
 float2 Dir;
 float Radius;
-float2 InvTexSize;
 
 float NormalDistribution1D(float X)
 {
@@ -16,15 +18,25 @@ float NormalDistribution1D(float X)
     return exp(-0.5 * X2 * X2);
 }
 
+float2 ClampInputUV(float2 UV)
+{
+    return clamp(UV, Input_UVViewportBilinearMin, Input_UVViewportBilinearMax);
+}
+
+float4 SampleInputTexture(float2 UV)
+{
+    return Input_Texture.SampleLevel(Input_Sampler, ClampInputUV(UV), 0.0);
+}
+
 float4 GaussianBlur(float2 UV)
 {
-    float2 Scale = Dir * InvTexSize;
+    float2 Scale = Dir * Input_ViewportSizeInverse;
     float4 Accum = 0.0;
     float Total = 0.0;
     float InvRadius = 1.0 / max(Radius, 1e-6);
 
     // Handle center first
-    Accum += InputTexture.SampleLevel(InputSampler, UV, 0.0);
+    Accum += SampleInputTexture(UV);
     Total += 1.0;
 
     // Determine how many pairs we can process symmetrically with bilinear taps
@@ -47,8 +59,8 @@ float4 GaussianBlur(float2 UV)
         float Offset = D1 - W2 / CombinedWeight;
         float2 UVOffset = Scale * Offset;
         
-        Accum += InputTexture.SampleLevel(InputSampler, UV + UVOffset, 0.0) * CombinedWeight;
-        Accum += InputTexture.SampleLevel(InputSampler, UV - UVOffset, 0.0) * CombinedWeight; 
+        Accum += SampleInputTexture(UV + UVOffset) * CombinedWeight;
+        Accum += SampleInputTexture(UV - UVOffset) * CombinedWeight;
         Total += CombinedWeight * 2.0;
     }
     
@@ -58,15 +70,20 @@ float4 GaussianBlur(float2 UV)
         float W = NormalDistribution1D(1.0 * InvRadius);
         float2 UVOffset = Scale * 1.0;
        
-        Accum += InputTexture.SampleLevel(InputSampler, UV + UVOffset, 0.0) * W;
-        Accum += InputTexture.SampleLevel(InputSampler, UV - UVOffset, 0.0) * W;
+        Accum += SampleInputTexture(UV + UVOffset) * W;
+        Accum += SampleInputTexture(UV - UVOffset) * W;
         Total += W * 2.0;
     }
     
     return Accum / max(Total, 1e-6);
 }
 
-float4 MainPS(float2 UV : TEXCOORD0) : SV_Target
+void MainPS(
+    float4 SvPosition : SV_POSITION,
+    out float4 OutColor : SV_Target0
+    )
 {
-    return GaussianBlur(UV);
-}
\ No newline at end of file
+    float2 InputUV = ApplyScreenTransform(SvPosition.xy, SvPositionToInputTextureUV);
+
+    OutColor = GaussianBlur(InputUV);
+}
diff --git a/Engine/Plugins/Compositing/Composite/Shaders/Private/CompositeCenteredScale.usf b/Engine/Plugins/Compositing/Composite/Shaders/Private/CompositeCenteredScale.usf
--- a/Engine/Plugins/Compositing/Composite/Shaders/Private/CompositeCenteredScale.usf
+++ b/Engine/Plugins/Compositing/Composite/Shaders/Private/CompositeCenteredScale.usf
@@ -1,19 +1,23 @@
 // Copyright Epic Games, Inc. All Rights Reserved.
 
 #include "/Engine/Private/Common.ush"
+#include "/Engine/Private/ScreenPass.ush"
 
-Texture2D InputTexture;
-SamplerState InputSampler;
+Texture2D Input_Texture;
+SamplerState Input_Sampler;
+FScreenTransform SvPositionToOutputViewportUV;
+FScreenTransform OutputViewportUVToInputTextureUV;
 float2 ScaleUV;
 static const float2 CenterUV = 0.5f;
 
 void MainPS(
-    noperspective float4 InUVAndScreenPos : TEXCOORD0,
+    float4 SvPosition : SV_POSITION,
     out float4 OutColor : SV_Target0
     )
 {
-    float2 UV = InUVAndScreenPos.xy;
-    UV = ((UV - CenterUV) * ScaleUV) + CenterUV;
+    float2 OutputViewportUV = ApplyScreenTransform(SvPosition.xy, SvPositionToOutputViewportUV);
+    float2 ScaledViewportUV = ((OutputViewportUV - CenterUV) * ScaleUV) + CenterUV;
+    float2 InputUV = ApplyScreenTransform(ScaledViewportUV, OutputViewportUVToInputTextureUV);
 
-    OutColor = InputTexture.Sample(InputSampler, UV);
+    OutColor = Input_Texture.Sample(Input_Sampler, InputUV);
 }
diff --git a/Engine/Plugins/Compositing/Composite/Shaders/Private/CompositeColorGrading.usf b/Engine/Plugins/Compositing/Composite/Shaders/Private/CompositeColorGrading.usf
--- a/Engine/Plugins/Compositing/Composite/Shaders/Private/CompositeColorGrading.usf
+++ b/Engine/Plugins/Compositing/Composite/Shaders/Private/CompositeColorGrading.usf
@@ -1,6 +1,7 @@
 // Copyright Epic Games, Inc. All Rights Reserved.
 
 #include "/Engine/Private/Common.ush"
+#include "/Engine/Private/ScreenPass.ush"
 #include "/Engine/Private/TonemapCommon.ush"
 #include "/Engine/Private/ACES/ACESCommon.ush"
 
@@ -10,8 +11,9 @@ static const float DefaultTemperature = 6500.0;
 static const float DefaultTint = 0.0;
 static const float Epsilon = 1e-8;
 
-Texture2D InputTexture;
-SamplerState InputSampler;
+Texture2D Input_Texture;
+SamplerState Input_Sampler;
+FScreenTransform SvPositionToInputTextureUV;
 
 uint bIsTemperatureWhiteBalance;
 float WhiteTemp;
@@ -100,11 +102,12 @@ float3 ColorCorrectAll( float3 WorkingColor )
 }
 
 void MainPS(
-    noperspective float4 InUVAndScreenPos : TEXCOORD0,
+    float4 SvPosition : SV_POSITION,
     out float4 OutColor : SV_Target0
     )
 {
-    float4 LinearColor = InputTexture.Sample(InputSampler, InUVAndScreenPos.xy);
+    float2 InputUV = ApplyScreenTransform(SvPosition.xy, SvPositionToInputTextureUV);
+    float4 LinearColor = Input_Texture.Sample(Input_Sampler, InputUV);
 	
 	if(bUnpremult)
 	{
diff --git a/Engine/Plugins/Compositing/Composite/Shaders/Private/CompositeColorKeyer.usf b/Engine/Plugins/Compositing/Composite/Shaders/Private/CompositeColorKeyer.usf
--- a/Engine/Plugins/Compositing/Composite/Shaders/Private/CompositeColorKeyer.usf
+++ b/Engine/Plugins/Compositing/Composite/Shaders/Private/CompositeColorKeyer.usf
@@ -9,16 +9,17 @@
     #error TILE_SIZE must be defined
 #endif
 
-Texture2D InputTexture;
-SamplerState InputSampler;
+Texture2D Input_Texture;
+SamplerState Input_Sampler;
 
-Texture2D CleanPlateTexture;
-SamplerState CleanPlateSampler;
+Texture2D CleanPlate_Texture;
+SamplerState CleanPlate_Sampler;
 
 RWTexture2D<float4> RWOutputTexture;
 
 SCREEN_PASS_TEXTURE_VIEWPORT(Input)
 SCREEN_PASS_TEXTURE_VIEWPORT(Output)
+FScreenTransform SvPositionToInputTextureUV;
 
 float3 KeyColor;
 float4 Params0;
@@ -63,9 +64,9 @@ void MainCS(
     uint2 DispatchThreadId : SV_DispatchThreadID)
 {
 	float4 SvPosition = float4((float2)DispatchThreadId + Output_ViewportMin + 0.5f, 0.0f, 1.0f);
-	float2 UV = (SvPosition.xy) * Output_ExtentInverse;
+	float2 OutputUV = (SvPosition.xy) * Output_ExtentInverse;
 
-	if (IsComputeUVOutOfBounds(UV))
+	if (IsComputeUVOutOfBounds(OutputUV))
 	{
 		return;
 	}
@@ -76,10 +77,11 @@ void MainCS(
 	const float AlphaMin = Params1.x;
 	const float AlphaMax = Params1.y;
 	const float DevignetteStrength = Params1.z;
+	float2 InputUV = ApplyScreenTransform(SvPosition.xy, SvPositionToInputTextureUV);
 	
-	float4 Foreground = InputTexture.SampleLevel(InputSampler, UV, 0);
+	float4 Foreground = Input_Texture.SampleLevel(Input_Sampler, InputUV, 0);
 #if USE_CLEAN_PLATE
-	float3 Background = CleanPlateTexture.SampleLevel(CleanPlateSampler, UV, 0).rgb;
+	float3 Background = CleanPlate_Texture.SampleLevel(CleanPlate_Sampler, InputUV, 0).rgb;
 #else
 	float3 Background = KeyColor;
 #endif
@@ -87,7 +89,7 @@ void MainCS(
 	float4 Result = Foreground;
 
 #if !USE_DESPILL_ONLY
-	float2 ScreenPos = UVToScreenPos(UV, Input_Extent);
+	float2 ScreenPos = UVToScreenPos(InputUV, Input_Extent);
 	float AspectRatio = Output_ViewportSize.y * Output_ViewportSizeInverse.x;
 	float VignetteMask = Vignette(ScreenPos, AspectRatio, DevignetteStrength);
 
diff --git a/Engine/Plugins/Compositing/Composite/Shaders/Private/CompositeDenoise.usf b/Engine/Plugins/Compositing/Composite/Shaders/Private/CompositeDenoise.usf
--- a/Engine/Plugins/Compositing/Composite/Shaders/Private/CompositeDenoise.usf
+++ b/Engine/Plugins/Compositing/Composite/Shaders/Private/CompositeDenoise.usf
@@ -1,6 +1,7 @@
 // Copyright Epic Games, Inc. All Rights Reserved.
 
 #include "/Engine/Private/Common.ush"
+#include "/Engine/Private/ScreenPass.ush"
 
 #include "CompositeColorKeyerCommon.ush"
 
@@ -16,9 +17,11 @@
 #define LDS_ARRAY_SIZE (LDS_TILE_WIDTH * LDS_TILE_WIDTH)
 #define THREADGROUP_TOTAL (TILE_SIZE * TILE_SIZE)
 
-Texture2D InputTexture;
+Texture2D Input_Texture;
+SamplerState Input_Sampler;
+SCREEN_PASS_TEXTURE_VIEWPORT(Input)
+
 RWTexture2D<float4> RWOutputTexture;
-uint2 Dimensions;
 
 groupshared float4 SharedColor[LDS_ARRAY_SIZE];
 
@@ -67,13 +70,13 @@ void MainCS(
 			uint2 TexelLocation =  GroupTexelOffset + uint2(
 				LinearGroupThreadId % LDS_TILE_WIDTH,
 				LinearGroupThreadId / LDS_TILE_WIDTH);
-			TexelLocation = min(TexelLocation, Dimensions - 1);
+			TexelLocation = min(TexelLocation, uint2(Input_ViewportSize) - 1);
 
 			if ((LinearGroupThreadId < LDS_ARRAY_SIZE) ||
 				(i != LoadCount - 1) ||
 				(LDS_ARRAY_SIZE % THREADGROUP_TOTAL) == 0)
 			{
-				SharedColor[LinearGroupThreadId] = InputTexture[TexelLocation];
+				SharedColor[LinearGroupThreadId] = Input_Texture[TexelLocation + uint2(Input_ViewportMin)];
 			}
 			LinearGroupThreadId += THREADGROUP_TOTAL;
 		}
@@ -82,7 +85,7 @@ void MainCS(
 	GroupMemoryBarrierWithGroupSync();
 
     // Discard out-of-bound texels
-    if (any(DispatchThreadId >= Dimensions))
+    if (any(DispatchThreadId >= uint2(Input_ViewportSize)))
     {
         return;
     }
diff --git a/Engine/Plugins/Compositing/Composite/Shaders/Private/CompositeDistortion.usf b/Engine/Plugins/Compositing/Composite/Shaders/Private/CompositeDistortion.usf
--- a/Engine/Plugins/Compositing/Composite/Shaders/Private/CompositeDistortion.usf
+++ b/Engine/Plugins/Compositing/Composite/Shaders/Private/CompositeDistortion.usf
@@ -4,9 +4,8 @@
 #include "/Engine/Private/LensDistortion.ush"
 #include "/Engine/Private/ScreenPass.ush"
 
-Texture2D InputTexture;
-SamplerState InputSampler;
-
+Texture2D Input_Texture;
+SamplerState Input_Sampler;
 SCREEN_PASS_TEXTURE_VIEWPORT(Input)
 SCREEN_PASS_TEXTURE_VIEWPORT(Output)
 
@@ -45,5 +44,5 @@ void MainPS(
 
     float2 Texture_UV = CalcInputUVs(ViewportUV, Input_ViewportSize, Input_ViewportMin, Input_ExtentInverse, DistortionUV);
 
-	OutColor = InputTexture.SampleLevel(InputSampler, Texture_UV, 0);
+	OutColor = Input_Texture.SampleLevel(Input_Sampler, Texture_UV, 0);
 }
diff --git a/Engine/Plugins/Compositing/Composite/Shaders/Private/CompositeMasking.usf b/Engine/Plugins/Compositing/Composite/Shaders/Private/CompositeMasking.usf
--- a/Engine/Plugins/Compositing/Composite/Shaders/Private/CompositeMasking.usf
+++ b/Engine/Plugins/Compositing/Composite/Shaders/Private/CompositeMasking.usf
@@ -1,6 +1,7 @@
 // Copyright Epic Games, Inc. All Rights Reserved.
 
 #include "/Engine/Private/Common.ush"
+#include "/Engine/Private/ScreenPass.ush"
 
 #include "CompositeCommon.ush"
 
@@ -15,17 +16,22 @@
 INPUT_PASS_TEXTURE(Input);
 INPUT_PASS_TEXTURE(Mask);
 
+FScreenTransform SvPositionToInputTextureUV;
+FScreenTransform SvPositionToMaskTextureUV;
+
 uint MaskSourceChannel;
 uint bUnpremult;
 uint bInvert;
 
 void MainPS(
-    noperspective float4 InUVAndScreenPos : TEXCOORD0,
+    float4 SvPosition : SV_POSITION,
     out float4 OutColor : SV_Target0
     )
 {
-    float4 InputColor = Input_Texture.SampleLevel(Input_Sampler, InUVAndScreenPos.xy, 0);
-	float4 MaskColor = Mask_Texture.SampleLevel(Mask_Sampler, InUVAndScreenPos.xy, 0);
+    float2 InputUV = ApplyScreenTransform(SvPosition.xy, SvPositionToInputTextureUV);
+    float2 MaskUV = ApplyScreenTransform(SvPosition.xy, SvPositionToMaskTextureUV);
+    float4 InputColor = Input_Texture.SampleLevel(Input_Sampler, InputUV, 0);
+	float4 MaskColor = Mask_Texture.SampleLevel(Mask_Sampler, MaskUV, 0);
 
     // Linearize source colors
     InputColor.rgb = Decode(InputColor.rgb, Input_SourceEncoding);
diff --git a/Engine/Plugins/Compositing/Composite/Source/Composite/Private/Layers/CompositeLayerMainRender.cpp b/Engine/Plugins/Compositing/Composite/Source/Composite/Private/Layers/CompositeLayerMainRender.cpp
--- a/Engine/Plugins/Compositing/Composite/Source/Composite/Private/Layers/CompositeLayerMainRender.cpp
+++ b/Engine/Plugins/Compositing/Composite/Source/Composite/Private/Layers/CompositeLayerMainRender.cpp
@@ -16,13 +16,18 @@ bool UCompositeLayerMainRender::GetProxy(const UE::CompositeCore::FPassInputDecl
 {
 	using namespace UE::CompositeCore;
 
+	FPassInputDecl PassInput;
 	FPassInputDeclArray Inputs;
 	Inputs.SetNum(FixedNumLayerInputs);
 
 	FPassInternalResourceDesc Desc0;
 	Desc0.Index = 0;
 	Desc0.bOriginalCopyBeforePasses = true;
-	Inputs[0].Set<FPassInternalResourceDesc>(Desc0);
+	PassInput.Set<FPassInternalResourceDesc>(Desc0);
+
+	AddChildPasses(PassInput, InContext, InFrameAllocator, LayerPasses);
+
+	Inputs[0] = PassInput;
 	Inputs[1] = GetDefaultSecondInput(InContext);
 
 	InContext.bNeedsSceneTextures = true;
diff --git a/Engine/Plugins/Compositing/Composite/Source/Composite/Private/Passes/CompositePassAlphaTransform.cpp b/Engine/Plugins/Compositing/Composite/Source/Composite/Private/Passes/CompositePassAlphaTransform.cpp
--- a/Engine/Plugins/Compositing/Composite/Source/Composite/Private/Passes/CompositePassAlphaTransform.cpp
+++ b/Engine/Plugins/Compositing/Composite/Source/Composite/Private/Passes/CompositePassAlphaTransform.cpp
@@ -3,6 +3,7 @@
 #include "Passes/CompositePassAlphaTransform.h"
 
 #include "Engine/Texture.h"
+#include "PixelShaderUtils.h"
 #include "RenderGraphBuilder.h"
 #include "RenderGraphUtils.h"
 #include "SceneView.h"
@@ -16,8 +17,8 @@ class FCompositePassAlphaTransformShader : public FGlobalShader
 	SHADER_USE_PARAMETER_STRUCT(FCompositePassAlphaTransformShader, FGlobalShader);
 
 	BEGIN_SHADER_PARAMETER_STRUCT(FParameters, )
-		SHADER_PARAMETER_RDG_TEXTURE(Texture2D, Texture)
-		SHADER_PARAMETER_SAMPLER(SamplerState, Sampler)
+		SHADER_PARAMETER_STRUCT(FScreenPassTextureInput, Input)
+		SHADER_PARAMETER(FScreenTransform, SvPositionToInputTextureUV)
 		SHADER_PARAMETER(uint32, PremultOp)
 		SHADER_PARAMETER(uint32, bInvertSourceAlpha)
 		SHADER_PARAMETER(uint32, OverrideOutputAlpha)
@@ -48,26 +49,26 @@ namespace UE
 					Output = CreateOutputRenderTarget(GraphBuilder, InView, PassContext.OutputViewRect, InputTexture.Texture->Desc, TEXT("AlphaTransformCompositePass"));
 				}
 
+				FRHISamplerState* BilinearClampSampler = TStaticSamplerState<SF_Bilinear>::GetRHI();
 				FGlobalShaderMap* GlobalShaderMap = GetGlobalShaderMap(InView.GetFeatureLevel());
-				FScreenPassTextureViewport Viewport = FScreenPassTextureViewport(Output);
-
 				FCompositePassAlphaTransformShader::FParameters* Parameters = GraphBuilder.AllocParameters<FCompositePassAlphaTransformShader::FParameters>();
-				Parameters->Texture = InputTexture.Texture;
-				Parameters->Sampler = TStaticSamplerState<SF_Bilinear>::GetRHI();
+				Parameters->Input = GetScreenPassTextureInput(InputTexture, BilinearClampSampler);
+				Parameters->SvPositionToInputTextureUV = (
+					FScreenTransform::ChangeTextureBasisFromTo(FScreenPassTextureViewport(Output), FScreenTransform::ETextureBasis::TexelPosition, FScreenTransform::ETextureBasis::ViewportUV) *
+					FScreenTransform::ChangeTextureBasisFromTo(FScreenPassTextureViewport(InputTexture), FScreenTransform::ETextureBasis::ViewportUV, FScreenTransform::ETextureBasis::TextureUV));
 				Parameters->PremultOp = static_cast<uint32>(PremultOp);
 				Parameters->bInvertSourceAlpha = static_cast<uint32>(bInvertSourceAlpha);
 				Parameters->OverrideOutputAlpha = static_cast<uint32>(OverrideOutputAlpha);
 				Parameters->RenderTargets[0] = Output.GetRenderTargetBinding();
 
 				TShaderMapRef<FCompositePassAlphaTransformShader> PixelShader(GlobalShaderMap);
-				AddDrawScreenPass(
+				FPixelShaderUtils::AddFullscreenPass(
 					GraphBuilder,
-					RDG_EVENT_NAME("Composite.AlphaTransform (%dx%d)", Viewport.Extent.X, Viewport.Extent.Y),
-					InView,
-					Viewport,
-					Viewport,
+					GlobalShaderMap,
+					RDG_EVENT_NAME("Composite.AlphaTransform (%dx%d)", Output.ViewRect.Width(), Output.ViewRect.Height()),
 					PixelShader,
-					Parameters
+					Parameters,
+					Output.ViewRect
 				);
 
 				return UE::CompositeCore::FPassTexture{ MoveTemp(Output), Metadata };
diff --git a/Engine/Plugins/Compositing/Composite/Source/Composite/Private/Passes/CompositePassBlur.cpp b/Engine/Plugins/Compositing/Composite/Source/Composite/Private/Passes/CompositePassBlur.cpp
--- a/Engine/Plugins/Compositing/Composite/Source/Composite/Private/Passes/CompositePassBlur.cpp
+++ b/Engine/Plugins/Compositing/Composite/Source/Composite/Private/Passes/CompositePassBlur.cpp
@@ -2,6 +2,7 @@
 
 #include "Passes/CompositePassBlur.h"
 
+#include "PixelShaderUtils.h"
 #include "RenderGraphBuilder.h"
 #include "RenderGraphUtils.h"
 #include "SceneView.h"
@@ -14,10 +15,9 @@ class FCompositePassBlurShader : public FGlobalShader
 	DECLARE_GLOBAL_SHADER(FCompositePassBlurShader);
 	SHADER_USE_PARAMETER_STRUCT(FCompositePassBlurShader, FGlobalShader);
 
-    BEGIN_SHADER_PARAMETER_STRUCT(FParameters, )
-    	SHADER_PARAMETER_RDG_TEXTURE(Texture2D, InputTexture)
-		SHADER_PARAMETER_SAMPLER(SamplerState, InputSampler)
-        SHADER_PARAMETER(FVector2f, InvTexSize)
+	BEGIN_SHADER_PARAMETER_STRUCT(FParameters, )
+		SHADER_PARAMETER_STRUCT(FScreenPassTextureInput, Input)
+		SHADER_PARAMETER(FScreenTransform, SvPositionToInputTextureUV)
 		SHADER_PARAMETER(FVector2f, Dir)
 		SHADER_PARAMETER(float, Radius)
 		RENDER_TARGET_BINDING_SLOTS()
@@ -58,48 +58,44 @@ namespace UE
 				FGlobalShaderMap* GlobalShaderMap = GetGlobalShaderMap(InView.GetFeatureLevel());
 				
 				{
-					FScreenPassTextureViewport Viewport = FScreenPassTextureViewport(IntermediateOutput);
 					FCompositePassBlurShader::FParameters* Parameters = GraphBuilder.AllocParameters<FCompositePassBlurShader::FParameters>();
-					Parameters->InputTexture = InputTexture.Texture;
-					Parameters->InputSampler = TStaticSamplerState<SF_Bilinear, AM_Clamp, AM_Clamp, AM_Clamp>::GetRHI();
-					const FIntVector InputSize = InputTexture.Texture->Desc.GetSize();
-					Parameters->InvTexSize = FVector2f::One() / FVector2f(FMath::Max(1.0f, (float)InputSize.X), FMath::Max(1.0f, (float)InputSize.Y));
+					Parameters->Input = GetScreenPassTextureInput(InputTexture, TStaticSamplerState<SF_Bilinear, AM_Clamp, AM_Clamp, AM_Clamp>::GetRHI());
+					Parameters->SvPositionToInputTextureUV = (
+						FScreenTransform::ChangeTextureBasisFromTo(FScreenPassTextureViewport(IntermediateOutput), FScreenTransform::ETextureBasis::TexelPosition, FScreenTransform::ETextureBasis::ViewportUV) *
+						FScreenTransform::ChangeTextureBasisFromTo(FScreenPassTextureViewport(InputTexture), FScreenTransform::ETextureBasis::ViewportUV, FScreenTransform::ETextureBasis::TextureUV));
 					Parameters->Dir = FVector2f(0, 1);
 					Parameters->Radius = Radius;
 					Parameters->RenderTargets[0] = IntermediateOutput.GetRenderTargetBinding();
 
 					TShaderMapRef<FCompositePassBlurShader> PixelShader(GlobalShaderMap);
-					AddDrawScreenPass(
+					FPixelShaderUtils::AddFullscreenPass(
 						GraphBuilder,
-						RDG_EVENT_NAME("Composite.Blur.Vertical (%dx%d)", Viewport.Extent.X, Viewport.Extent.Y),
-						InView,
-						Viewport,
-						Viewport,
+						GlobalShaderMap,
+						RDG_EVENT_NAME("Composite.Blur.Vertical (%dx%d)", IntermediateOutput.ViewRect.Width(), IntermediateOutput.ViewRect.Height()),
 						PixelShader,
-						Parameters
+						Parameters,
+						IntermediateOutput.ViewRect
 					);
 				}
 
 				{
-					FScreenPassTextureViewport Viewport = FScreenPassTextureViewport(Output);
 					FCompositePassBlurShader::FParameters* Parameters = GraphBuilder.AllocParameters<FCompositePassBlurShader::FParameters>();
-					Parameters->InputTexture = IntermediateOutput.Texture;
-					Parameters->InputSampler = TStaticSamplerState<SF_Bilinear, AM_Clamp, AM_Clamp, AM_Clamp>::GetRHI();
-					const FIntVector InputSize = IntermediateOutput.Texture->Desc.GetSize();
-					Parameters->InvTexSize = FVector2f(1.0f / FMath::Max(1.0f, (float)InputSize.X), 1.0f / FMath::Max(1.0f, (float)InputSize.Y));
+					Parameters->Input = GetScreenPassTextureInput(IntermediateOutput, TStaticSamplerState<SF_Bilinear, AM_Clamp, AM_Clamp, AM_Clamp>::GetRHI());
+					Parameters->SvPositionToInputTextureUV = (
+						FScreenTransform::ChangeTextureBasisFromTo(FScreenPassTextureViewport(Output), FScreenTransform::ETextureBasis::TexelPosition, FScreenTransform::ETextureBasis::ViewportUV) *
+						FScreenTransform::ChangeTextureBasisFromTo(FScreenPassTextureViewport(IntermediateOutput), FScreenTransform::ETextureBasis::ViewportUV, FScreenTransform::ETextureBasis::TextureUV));
 					Parameters->Dir = FVector2f(1, 0);
 					Parameters->Radius = Radius;
 					Parameters->RenderTargets[0] = Output.GetRenderTargetBinding();
 
 					TShaderMapRef<FCompositePassBlurShader> PixelShader(GlobalShaderMap);
-					AddDrawScreenPass(
+					FPixelShaderUtils::AddFullscreenPass(
 						GraphBuilder,
-						RDG_EVENT_NAME("Composite.Blur.Horizontal (%dx%d)", Viewport.Extent.X, Viewport.Extent.Y),
-						InView,
-						Viewport,
-						Viewport,
+						GlobalShaderMap,
+						RDG_EVENT_NAME("Composite.Blur.Horizontal (%dx%d)", Output.ViewRect.Width(), Output.ViewRect.Height()),
 						PixelShader,
-						Parameters
+						Parameters,
+						Output.ViewRect
 					);
 				}
 
diff --git a/Engine/Plugins/Compositing/Composite/Source/Composite/Private/Passes/CompositePassCenteredScale.cpp b/Engine/Plugins/Compositing/Composite/Source/Composite/Private/Passes/CompositePassCenteredScale.cpp
--- a/Engine/Plugins/Compositing/Composite/Source/Composite/Private/Passes/CompositePassCenteredScale.cpp
+++ b/Engine/Plugins/Compositing/Composite/Source/Composite/Private/Passes/CompositePassCenteredScale.cpp
@@ -8,6 +8,8 @@
 #include "Layers/CompositeLayerPlate.h"
 #include "MediaTexture.h"
 #include "Passes/CompositeCorePassProxy.h"
+#include "PixelShaderUtils.h"
+#include "ScreenPass.h"
 
 DECLARE_GPU_STAT_NAMED(FCompositeCenteredScale, TEXT("Composite.CenteredScale"));
 
@@ -17,8 +19,9 @@ class FCompositeCenteredScaleShader : public FGlobalShader
 	SHADER_USE_PARAMETER_STRUCT(FCompositeCenteredScaleShader, FGlobalShader);
 
 	BEGIN_SHADER_PARAMETER_STRUCT(FParameters, )
-		SHADER_PARAMETER_RDG_TEXTURE(Texture2D, InputTexture)
-		SHADER_PARAMETER_SAMPLER(SamplerState, InputSampler)
+		SHADER_PARAMETER_STRUCT(FScreenPassTextureInput, Input)
+		SHADER_PARAMETER(FScreenTransform, SvPositionToOutputViewportUV)
+		SHADER_PARAMETER(FScreenTransform, OutputViewportUVToInputTextureUV)
 		SHADER_PARAMETER(FVector2f, ScaleUV)
 		RENDER_TARGET_BINDING_SLOTS()
 	END_SHADER_PARAMETER_STRUCT()
@@ -54,23 +57,27 @@ namespace UE
 				}
 
 				FGlobalShaderMap* GlobalShaderMap = GetGlobalShaderMap(InView.GetFeatureLevel());
-				FScreenPassTextureViewport Viewport = FScreenPassTextureViewport(Output);
-
 				FCompositeCenteredScaleShader::FParameters* Parameters = GraphBuilder.AllocParameters<FCompositeCenteredScaleShader::FParameters>();
-				Parameters->InputTexture = Input.Texture;
-				Parameters->InputSampler = TStaticSamplerState<SF_Bilinear>::GetRHI();
+				Parameters->Input = GetScreenPassTextureInput(Input, TStaticSamplerState<SF_Bilinear>::GetRHI());
+				Parameters->SvPositionToOutputViewportUV = FScreenTransform::ChangeTextureBasisFromTo(
+					FScreenPassTextureViewport(Output),
+					FScreenTransform::ETextureBasis::TexelPosition,
+					FScreenTransform::ETextureBasis::ViewportUV);
+				Parameters->OutputViewportUVToInputTextureUV = FScreenTransform::ChangeTextureBasisFromTo(
+					FScreenPassTextureViewport(Input),
+					FScreenTransform::ETextureBasis::ViewportUV,
+					FScreenTransform::ETextureBasis::TextureUV);
 				Parameters->ScaleUV = ScaleUV;
 				Parameters->RenderTargets[0] = Output.GetRenderTargetBinding();
 
 				TShaderMapRef<FCompositeCenteredScaleShader> PixelShader(GlobalShaderMap);
-				AddDrawScreenPass(
+				FPixelShaderUtils::AddFullscreenPass(
 					GraphBuilder,
-					RDG_EVENT_NAME("Composite.CenteredScale (%dx%d)", Viewport.Extent.X, Viewport.Extent.Y),
-					InView,
-					Viewport,
-					Viewport,
+					GlobalShaderMap,
+					RDG_EVENT_NAME("Composite.CenteredScale (%dx%d)", Output.ViewRect.Width(), Output.ViewRect.Height()),
 					PixelShader,
-					Parameters
+					Parameters,
+					Output.ViewRect
 				);
 
 				return UE::CompositeCore::FPassTexture{ MoveTemp(Output), Metadata };
diff --git a/Engine/Plugins/Compositing/Composite/Source/Composite/Private/Passes/CompositePassColorGrading.cpp b/Engine/Plugins/Compositing/Composite/Source/Composite/Private/Passes/CompositePassColorGrading.cpp
--- a/Engine/Plugins/Compositing/Composite/Source/Composite/Private/Passes/CompositePassColorGrading.cpp
+++ b/Engine/Plugins/Compositing/Composite/Source/Composite/Private/Passes/CompositePassColorGrading.cpp
@@ -2,6 +2,7 @@
 
 #include "Passes/CompositePassColorGrading.h"
 
+#include "PixelShaderUtils.h"
 #include "SceneManagement.h"
 #include "ScreenPass.h"
 #include "PostProcess/PostProcessMaterialInputs.h"
@@ -23,8 +24,7 @@ class FCompositeColorGradingShader : public FGlobalShader
 	SHADER_USE_PARAMETER_STRUCT(FCompositeColorGradingShader, FGlobalShader);
 
 	BEGIN_SHADER_PARAMETER_STRUCT(FParameters, )
-		SHADER_PARAMETER_RDG_TEXTURE(Texture2D, InputTexture)
-		SHADER_PARAMETER_SAMPLER(SamplerState, InputSampler)
+		SHADER_PARAMETER_STRUCT(FScreenPassTextureInput, Input)
 		SHADER_PARAMETER_STRUCT_REF(FWorkingColorSpaceShaderParameters, WorkingColorSpace)
 		SHADER_PARAMETER_STRUCT(FColorGradingPerRangeParameters, Global)
 		SHADER_PARAMETER_STRUCT(FColorGradingPerRangeParameters, Shadows)
@@ -37,6 +37,7 @@ class FCompositeColorGradingShader : public FGlobalShader
 		SHADER_PARAMETER(uint32, bIsTemperatureWhiteBalance)
 		SHADER_PARAMETER(float, WhiteTemp)
 		SHADER_PARAMETER(float, WhiteTint)
+		SHADER_PARAMETER(FScreenTransform, SvPositionToInputTextureUV)
 		RENDER_TARGET_BINDING_SLOTS()
 	END_SHADER_PARAMETER_STRUCT()
 
@@ -78,11 +79,8 @@ namespace UE::Composite::Private
 		}
 
 		FGlobalShaderMap* GlobalShaderMap = GetGlobalShaderMap(InView.GetFeatureLevel());
-		FScreenPassTextureViewport Viewport = FScreenPassTextureViewport(Output);
-
 		FCompositeColorGradingShader::FParameters* Parameters = GraphBuilder.AllocParameters<FCompositeColorGradingShader::FParameters>();
-		Parameters->InputTexture = Input.Texture;
-		Parameters->InputSampler = TStaticSamplerState<SF_Bilinear>::GetRHI();
+		Parameters->Input = GetScreenPassTextureInput(Input, TStaticSamplerState<SF_Bilinear>::GetRHI());
 		Parameters->WorkingColorSpace = GDefaultWorkingColorSpaceUniformBuffer.GetUniformBufferRef();
 		Parameters->Global = FCompositeColorGradingShader::ToParameters(ColorGradingSettings.Global);
 		Parameters->Shadows = FCompositeColorGradingShader::ToParameters(ColorGradingSettings.Shadows);
@@ -95,17 +93,19 @@ namespace UE::Composite::Private
 		Parameters->bIsTemperatureWhiteBalance = (TemperatureSettings.TemperatureType == ETemperatureMethod::TEMP_WhiteBalance);
 		Parameters->WhiteTemp = TemperatureSettings.WhiteTemp;
 		Parameters->WhiteTint = TemperatureSettings.WhiteTint;
+		Parameters->SvPositionToInputTextureUV = (
+			FScreenTransform::ChangeTextureBasisFromTo(FScreenPassTextureViewport(Output), FScreenTransform::ETextureBasis::TexelPosition, FScreenTransform::ETextureBasis::ViewportUV) *
+			FScreenTransform::ChangeTextureBasisFromTo(FScreenPassTextureViewport(Input), FScreenTransform::ETextureBasis::ViewportUV, FScreenTransform::ETextureBasis::TextureUV));
 		Parameters->RenderTargets[0] = Output.GetRenderTargetBinding();
 
 		TShaderMapRef<FCompositeColorGradingShader> PixelShader(GlobalShaderMap);
-		AddDrawScreenPass(
+		FPixelShaderUtils::AddFullscreenPass(
 			GraphBuilder,
-			RDG_EVENT_NAME("Composite.ColorGrading (%dx%d)", Viewport.Extent.X, Viewport.Extent.Y),
-			InView,
-			Viewport,
-			Viewport,
+			GlobalShaderMap,
+			RDG_EVENT_NAME("Composite.ColorGrading (%dx%d)", Output.ViewRect.Width(), Output.ViewRect.Height()),
 			PixelShader,
-			Parameters
+			Parameters,
+			Output.ViewRect
 		);
 
 		return CompositeCore::FPassTexture{ MoveTemp(Output), Metadata };
diff --git a/Engine/Plugins/Compositing/Composite/Source/Composite/Private/Passes/CompositePassColorKeyer.cpp b/Engine/Plugins/Compositing/Composite/Source/Composite/Private/Passes/CompositePassColorKeyer.cpp
--- a/Engine/Plugins/Compositing/Composite/Source/Composite/Private/Passes/CompositePassColorKeyer.cpp
+++ b/Engine/Plugins/Compositing/Composite/Source/Composite/Private/Passes/CompositePassColorKeyer.cpp
@@ -5,6 +5,7 @@
 #include "Engine/Texture.h"
 #include "Passes/CompositeCorePassProxy.h"
 #include "PixelShaderUtils.h"
+#include "ScreenPass.h"
 #include "SystemTextures.h"
 #include "TextureResource.h"
 #include "UObject/FortniteMainBranchObjectVersion.h"
@@ -22,9 +23,8 @@ class FCompositePassDenoiseShader : public FGlobalShader
 	SHADER_USE_PARAMETER_STRUCT(FCompositePassDenoiseShader, FGlobalShader);
 
 	BEGIN_SHADER_PARAMETER_STRUCT(FParameters, )
-		SHADER_PARAMETER_RDG_TEXTURE(Texture2D, InputTexture)
+		SHADER_PARAMETER_STRUCT(FScreenPassTextureInput, Input)
 		SHADER_PARAMETER_RDG_TEXTURE_UAV(RWTexture2D<float4>, RWOutputTexture)
-		SHADER_PARAMETER(FIntPoint, Dimensions)
 	END_SHADER_PARAMETER_STRUCT()
 
 	static inline void ModifyCompilationEnvironment(const FGlobalShaderPermutationParameters& Parameters, FShaderCompilerEnvironment& OutEnvironment)
@@ -47,13 +47,11 @@ class FCompositePassColorKeyerShader : public FGlobalShader
 	using FPermutationDomain = TShaderPermutationDomain<FUseCleanPlate, FDespillOnly>;
 
 	BEGIN_SHADER_PARAMETER_STRUCT(FParameters, )
-		SHADER_PARAMETER_RDG_TEXTURE(Texture2D, InputTexture)
-		SHADER_PARAMETER_SAMPLER(SamplerState, InputSampler)
-		SHADER_PARAMETER_RDG_TEXTURE(Texture2D, CleanPlateTexture)
-		SHADER_PARAMETER_SAMPLER(SamplerState, CleanPlateSampler)
+		SHADER_PARAMETER_STRUCT(FScreenPassTextureInput, Input)
+		SHADER_PARAMETER_STRUCT(FScreenPassTextureInput, CleanPlate)
 		SHADER_PARAMETER_RDG_TEXTURE_UAV(RWTexture2D<float4>, RWOutputTexture)
-		SHADER_PARAMETER_STRUCT(FScreenPassTextureViewportParameters, Input)
 		SHADER_PARAMETER_STRUCT(FScreenPassTextureViewportParameters, Output)
+		SHADER_PARAMETER(FScreenTransform, SvPositionToInputTextureUV)
 		SHADER_PARAMETER(FVector3f, KeyColor)
 		SHADER_PARAMETER(FVector4f, Params0)
 		SHADER_PARAMETER(FVector4f, Params1)
@@ -92,15 +90,14 @@ namespace UE
 				RDG_EVENT_SCOPE_STAT(GraphBuilder, FCompositeDenoise, "Composite.Denoise");
 
 				FGlobalShaderMap* GlobalShaderMap = GetGlobalShaderMap(FeatureLevel);
-				const FIntPoint TextureSize = Input.Texture->Desc.Extent;
+				const FIntPoint TextureSize = Input.ViewRect.Size();
 
 				const FRDGTextureDesc OutputDesc = FRDGTextureDesc::Create2D(TextureSize, PF_FloatRGBA, FClearValueBinding::Black, TexCreate_UAV | TexCreate_ShaderResource | TexCreate_RenderTargetable);
-				FScreenPassRenderTarget Output = FScreenPassRenderTarget(GraphBuilder.CreateTexture(OutputDesc, TEXT("CompositeTextureDenoised")), ERenderTargetLoadAction::ENoAction);
+				FScreenPassRenderTarget Output = FScreenPassRenderTarget(GraphBuilder.CreateTexture(OutputDesc, TEXT("CompositeTextureDenoised")), FIntRect(FIntPoint::ZeroValue, TextureSize), ERenderTargetLoadAction::ENoAction);
 
 				FCompositePassDenoiseShader::FParameters* PassParameters = GraphBuilder.AllocParameters<FCompositePassDenoiseShader::FParameters>();
-				PassParameters->InputTexture = Input.Texture;
+				PassParameters->Input = GetScreenPassTextureInput(Input, TStaticSamplerState<SF_Point>::GetRHI());
 				PassParameters->RWOutputTexture = GraphBuilder.CreateUAV(Output.Texture);
-				PassParameters->Dimensions = TextureSize;
 
 				TShaderMapRef<FCompositePassDenoiseShader> ComputeShader(GlobalShaderMap);
 				FComputeShaderUtils::AddPass(
@@ -139,15 +136,15 @@ namespace UE
 				FGlobalShaderMap* GlobalShaderMap = GetGlobalShaderMap(InView.GetFeatureLevel());
 
 				const bool bSameSize = Output.ViewRect.Size() == Input.ViewRect.Size();
-				const FScreenPassTextureViewportParameters InputParameters = GetScreenPassTextureViewportParameters(FScreenPassTextureViewport(Input));
 				const FScreenPassTextureViewportParameters OutputParameters = GetScreenPassTextureViewportParameters(FScreenPassTextureViewport(Output));
 
 				FCompositePassColorKeyerShader::FParameters* PassParameters = GraphBuilder.AllocParameters<FCompositePassColorKeyerShader::FParameters>();
-				PassParameters->InputTexture = Input.Texture;
-				PassParameters->InputSampler = bSameSize ? TStaticSamplerState<SF_Point>::GetRHI() : TStaticSamplerState<SF_Bilinear>::GetRHI();
+				PassParameters->Input = GetScreenPassTextureInput(Input, bSameSize ? TStaticSamplerState<SF_Point>::GetRHI() : TStaticSamplerState<SF_Bilinear>::GetRHI());
 				PassParameters->RWOutputTexture = GraphBuilder.CreateUAV(Output.Texture);
-				PassParameters->Input = InputParameters;
 				PassParameters->Output = OutputParameters;
+				PassParameters->SvPositionToInputTextureUV = (
+					FScreenTransform::ChangeTextureBasisFromTo(FScreenPassTextureViewport(Output), FScreenTransform::ETextureBasis::TexelPosition, FScreenTransform::ETextureBasis::ViewportUV) *
+					FScreenTransform::ChangeTextureBasisFromTo(FScreenPassTextureViewport(Input), FScreenTransform::ETextureBasis::ViewportUV, FScreenTransform::ETextureBasis::TextureUV));
 				PassParameters->KeyColor = FVector3f(KeyColor);
 				PassParameters->Params0 = Params0;
 				PassParameters->Params1 = Params1;
@@ -157,8 +154,8 @@ namespace UE
 				PassParameters->bInvertAlpha = static_cast<uint32>(bInvertAlpha);
 	
 				FRDGTextureRef CleanPlateResource = nullptr;
-				PassParameters->CleanPlateTexture = GSystemTextures.GetBlackDummy(GraphBuilder);
-				PassParameters->CleanPlateSampler = TStaticSamplerState<SF_Bilinear>::GetRHI();
+				PassParameters->CleanPlate.Texture = GSystemTextures.GetBlackDummy(GraphBuilder);
+				PassParameters->CleanPlate.Sampler = TStaticSamplerState<SF_Bilinear>::GetRHI();
 
 				TStrongObjectPtr<UTexture> CleanPlateTexturePtr = CleanPlateWeakPtr.Pin();
 				if (CleanPlateTexturePtr.IsValid())
@@ -172,7 +169,7 @@ namespace UE
 							CleanPlateResource = GraphBuilder.RegisterExternalTexture(CreateRenderTarget(TextureRHI, TEXT("CompositeCleanPlateTexture")));
 							GraphBuilder.UseInternalAccessMode(CleanPlateResource);
 
-							PassParameters->CleanPlateTexture = CleanPlateResource;
+							PassParameters->CleanPlate.Texture = CleanPlateResource;
 						}
 					}
 				}
diff --git a/Engine/Plugins/Compositing/Composite/Source/Composite/Private/Passes/CompositePassDistortion.cpp b/Engine/Plugins/Compositing/Composite/Source/Composite/Private/Passes/CompositePassDistortion.cpp
--- a/Engine/Plugins/Compositing/Composite/Source/Composite/Private/Passes/CompositePassDistortion.cpp
+++ b/Engine/Plugins/Compositing/Composite/Source/Composite/Private/Passes/CompositePassDistortion.cpp
@@ -20,9 +20,7 @@ class FCompositeDistortionShader : public FGlobalShader
 	SHADER_USE_PARAMETER_STRUCT(FCompositeDistortionShader, FGlobalShader);
 
 	BEGIN_SHADER_PARAMETER_STRUCT(FParameters, )
-		SHADER_PARAMETER_RDG_TEXTURE(Texture2D, InputTexture)
-		SHADER_PARAMETER_SAMPLER(SamplerState, InputSampler)
-		SHADER_PARAMETER_STRUCT(FScreenPassTextureViewportParameters, Input)
+		SHADER_PARAMETER_STRUCT(FScreenPassTextureInput, Input)
 		SHADER_PARAMETER_STRUCT(FScreenPassTextureViewportParameters, Output)
 		SHADER_PARAMETER_RDG_TEXTURE(Texture2D, DistortingDisplacementTexture)
 		SHADER_PARAMETER_SAMPLER(SamplerState, DistortingDisplacementSampler)
@@ -92,17 +90,12 @@ public:
 
 		if (bValidLensDistortion)
 		{
-			const FScreenPassTextureViewportParameters InputParameters = GetScreenPassTextureViewportParameters(FScreenPassTextureViewport(Input));
-			const FScreenPassTextureViewportParameters OutputParameters = GetScreenPassTextureViewportParameters(FScreenPassTextureViewport(Output));
-
 			FGlobalShaderMap* GlobalShaderMap = GetGlobalShaderMap(InView.GetFeatureLevel());
 			FScreenPassTextureViewport Viewport = FScreenPassTextureViewport(Output);
 
 			FCompositeDistortionShader::FParameters* PassParameters = GraphBuilder.AllocParameters<FCompositeDistortionShader::FParameters>();
-			PassParameters->InputTexture = Input.Texture;
-			PassParameters->InputSampler = TStaticSamplerState<SF_Bilinear>::GetRHI();
-			PassParameters->Input = InputParameters;
-			PassParameters->Output = OutputParameters;
+			PassParameters->Input = GetScreenPassTextureInput(Input, TStaticSamplerState<SF_Bilinear>::GetRHI());
+			PassParameters->Output = GetScreenPassTextureViewportParameters(FScreenPassTextureViewport(Output));
 			if (ViewDistortionLUT.IsEnabled())
 			{
 				PassParameters->DistortingDisplacementTexture = ViewDistortionLUT.DistortingDisplacementTexture;
diff --git a/Engine/Plugins/Compositing/Composite/Source/Composite/Private/Passes/CompositePassMasking.cpp b/Engine/Plugins/Compositing/Composite/Source/Composite/Private/Passes/CompositePassMasking.cpp
--- a/Engine/Plugins/Compositing/Composite/Source/Composite/Private/Passes/CompositePassMasking.cpp
+++ b/Engine/Plugins/Compositing/Composite/Source/Composite/Private/Passes/CompositePassMasking.cpp
@@ -4,6 +4,7 @@
 #include "Passes/CompositeCorePassProxy.h"
 
 #include "Engine/Texture.h"
+#include "PixelShaderUtils.h"
 #include "RenderGraphBuilder.h"
 #include "RenderGraphUtils.h"
 #include "SceneView.h"
@@ -29,6 +30,8 @@ class FCompositePassMaskingShader : public FGlobalShader
 		SHADER_PARAMETER(uint32, MaskSourceChannel)
 		SHADER_PARAMETER(uint32, bUnpremult)
 		SHADER_PARAMETER(uint32, bInvert)
+		SHADER_PARAMETER(FScreenTransform, SvPositionToInputTextureUV)
+		SHADER_PARAMETER(FScreenTransform, SvPositionToMaskTextureUV)
 		RENDER_TARGET_BINDING_SLOTS()
 	END_SHADER_PARAMETER_STRUCT()
 };
@@ -53,7 +56,7 @@ namespace UE
 				FCompositePassMaskingShader::FInputPassTextureParameters Parameters;
 				Parameters.Texture = Input.Texture.Texture;
 				Parameters.Sampler = TStaticSamplerState<SF_Bilinear>::GetRHI();
-				Parameters.bInvertedAlpha = Input.Metadata.bInvertedAlpha;
+				Parameters.bInvertedAlpha = static_cast<uint32>(Input.Metadata.bInvertedAlpha);
 				Parameters.SourceEncoding = static_cast<uint32>(Input.Metadata.Encoding);
 
 				return Parameters;
@@ -76,25 +79,28 @@ namespace UE
 				}
 
 				FGlobalShaderMap* GlobalShaderMap = GetGlobalShaderMap(InView.GetFeatureLevel());
-				FScreenPassTextureViewport Viewport = FScreenPassTextureViewport(Output);
-
 				FCompositePassMaskingShader::FParameters* Parameters = GraphBuilder.AllocParameters<FCompositePassMaskingShader::FParameters>();
 				Parameters->Input = GetPassTextureParameters(Inputs[0]);
 				Parameters->Mask = GetPassTextureParameters(Inputs[1]);
 				Parameters->MaskSourceChannel = MaskSourceChannel;
-				Parameters->bUnpremult = bUnpremult;
-				Parameters->bInvert = bInvert;
+				Parameters->bUnpremult = static_cast<uint32>(bUnpremult);
+				Parameters->bInvert = static_cast<uint32>(bInvert);
+				Parameters->SvPositionToInputTextureUV = (
+					FScreenTransform::ChangeTextureBasisFromTo(FScreenPassTextureViewport(Output), FScreenTransform::ETextureBasis::TexelPosition, FScreenTransform::ETextureBasis::ViewportUV) *
+					FScreenTransform::ChangeTextureBasisFromTo(FScreenPassTextureViewport(Inputs[0].Texture), FScreenTransform::ETextureBasis::ViewportUV, FScreenTransform::ETextureBasis::TextureUV));
+				Parameters->SvPositionToMaskTextureUV = (
+					FScreenTransform::ChangeTextureBasisFromTo(FScreenPassTextureViewport(Output), FScreenTransform::ETextureBasis::TexelPosition, FScreenTransform::ETextureBasis::ViewportUV) *
+					FScreenTransform::ChangeTextureBasisFromTo(FScreenPassTextureViewport(Inputs[1].Texture), FScreenTransform::ETextureBasis::ViewportUV, FScreenTransform::ETextureBasis::TextureUV));
 				Parameters->RenderTargets[0] = Output.GetRenderTargetBinding();
 
 				TShaderMapRef<FCompositePassMaskingShader> PixelShader(GlobalShaderMap);
-				AddDrawScreenPass(
+				FPixelShaderUtils::AddFullscreenPass(
 					GraphBuilder,
-					RDG_EVENT_NAME("Composite.Masking (%dx%d)", Viewport.Extent.X, Viewport.Extent.Y),
-					InView,
-					Viewport,
-					Viewport,
+					GlobalShaderMap,
+					RDG_EVENT_NAME("Composite.Masking (%dx%d)", Output.ViewRect.Width(), Output.ViewRect.Height()),
 					PixelShader,
-					Parameters
+					Parameters,
+					Output.ViewRect
 				);
 
 				return UE::CompositeCore::FPassTexture{ MoveTemp(Output), Metadata };
diff --git a/Engine/Plugins/Compositing/Composite/Source/Composite/Public/Layers/CompositeLayerMainRender.h b/Engine/Plugins/Compositing/Composite/Source/Composite/Public/Layers/CompositeLayerMainRender.h
--- a/Engine/Plugins/Compositing/Composite/Source/Composite/Public/Layers/CompositeLayerMainRender.h
+++ b/Engine/Plugins/Compositing/Composite/Source/Composite/Public/Layers/CompositeLayerMainRender.h
@@ -25,6 +25,10 @@ public:
 
 	/** Getter function to override, returning pass layer proxies to be passed to the render thread. (Proxy objects should be allocated from the provided allocator.) */
 	UE_API virtual bool GetProxy(const UE::CompositeCore::FPassInputDecl& InputDecl, FCompositeTraversalContext& InContext, FSceneRenderingBulkObjectAllocator& InFrameAllocator, FCompositeCorePassProxy*& OutProxy) const override;
+
+	/** Sub-passes applied onto the main render input before it is merged as this layer. */
+	UPROPERTY(EditAnywhere, BlueprintReadWrite, Instanced, Category = "", meta = (DisallowedClasses = "/Script/Composite.CompositeLayerBase, /Script/Composite.CompositePassDistortion"))
+	TArray<TObjectPtr<UCompositePassBase>> LayerPasses;
 };
 
 #undef UE_API
diff --git a/Engine/Plugins/Compositing/CompositeCore/Shaders/Private/CompositeCoreDisplayTransform.usf b/Engine/Plugins/Compositing/CompositeCore/Shaders/Private/CompositeCoreDisplayTransform.usf
--- a/Engine/Plugins/Compositing/CompositeCore/Shaders/Private/CompositeCoreDisplayTransform.usf
+++ b/Engine/Plugins/Compositing/CompositeCore/Shaders/Private/CompositeCoreDisplayTransform.usf
@@ -1,9 +1,12 @@
 // Copyright Epic Games, Inc. All Rights Reserved.
 
 #include "/Engine/Private/Common.ush"
+#include "/Engine/Private/ScreenPass.ush"
+
+FScreenTransform SvPositionToInputTextureUV;
+Texture2D Input_Texture;
+SamplerState Input_Sampler;
 
-Texture2D InputTexture;
-SamplerState InputSampler;
 uint bIsForward;
 float Gamma;
 float InvGamma;
@@ -36,11 +39,12 @@ float3 GammaDecode(float3 Color)
 }
 
 void MainPS(
-    noperspective float4 InUVAndScreenPos : TEXCOORD0,
+    float4 SvPosition : SV_POSITION,
     out float4 OutColor : SV_Target0
     )
 {
-    OutColor = InputTexture.Sample(InputSampler, InUVAndScreenPos.xy);
+    float2 TextureUV = ApplyScreenTransform(SvPosition.xy, SvPositionToInputTextureUV);
+    OutColor = Input_Texture.Sample(Input_Sampler, TextureUV);
 
 #if VERSION == 0
     if(bIsForward)
diff --git a/Engine/Plugins/Compositing/CompositeCore/Source/CompositeCore/Private/Passes/CompositeCorePassFXAAProxy.cpp b/Engine/Plugins/Compositing/CompositeCore/Source/CompositeCore/Private/Passes/CompositeCorePassFXAAProxy.cpp
--- a/Engine/Plugins/Compositing/CompositeCore/Source/CompositeCore/Private/Passes/CompositeCorePassFXAAProxy.cpp
+++ b/Engine/Plugins/Compositing/CompositeCore/Source/CompositeCore/Private/Passes/CompositeCorePassFXAAProxy.cpp
@@ -2,6 +2,7 @@
 
 #include "Passes/CompositeCorePassFXAAProxy.h"
 
+#include "PixelShaderUtils.h"
 #include "RenderGraphBuilder.h"
 #include "RenderGraphUtils.h"
 #include "SceneView.h"
@@ -18,9 +19,8 @@ class FCompositeCoreDisplayTransformShader : public FGlobalShader
 	SHADER_USE_PARAMETER_STRUCT(FCompositeCoreDisplayTransformShader, FGlobalShader);
 
 	BEGIN_SHADER_PARAMETER_STRUCT(FParameters, )
-		SHADER_PARAMETER_RDG_TEXTURE(Texture2D, InputTexture)
-		SHADER_PARAMETER_SAMPLER(SamplerState, InputSampler)
-		SHADER_PARAMETER(FIntPoint, Dimensions)
+		SHADER_PARAMETER_STRUCT(FScreenPassTextureInput, Input)
+		SHADER_PARAMETER(FScreenTransform, SvPositionToInputTextureUV)
 		SHADER_PARAMETER(uint32, bIsForward)
 		SHADER_PARAMETER(float, Gamma)
 		SHADER_PARAMETER(float, InvGamma)
@@ -48,30 +48,31 @@ namespace UE
 					FRDGTextureDesc OutputDesc = Input.Texture->Desc;
 					OutputDesc.Flags |= TexCreate_RenderTargetable;
 					OutputDesc.Format = FSceneTexturesConfig::Get().ColorFormat;
+					OutputDesc.Extent = Input.ViewRect.Size();
 
-					Output = FScreenPassRenderTarget(GraphBuilder.CreateTexture(OutputDesc, TEXT("FXAACompositePass")), ERenderTargetLoadAction::ENoAction);
+					const FIntRect ZeroBasedViewRect(FIntPoint::ZeroValue, Input.ViewRect.Size());
+					Output = FScreenPassRenderTarget(GraphBuilder.CreateTexture(OutputDesc, TEXT("FXAACompositePass")), ZeroBasedViewRect, ERenderTargetLoadAction::ENoAction);
 				}
 
-				FScreenPassTextureViewport Viewport = FScreenPassTextureViewport(Output);
-
 				FCompositeCoreDisplayTransformShader::FParameters* Parameters = GraphBuilder.AllocParameters<FCompositeCoreDisplayTransformShader::FParameters>();
-				Parameters->InputTexture = Input.Texture;
-				Parameters->InputSampler = TStaticSamplerState<>::GetRHI();
+				Parameters->Input = GetScreenPassTextureInput(Input, TStaticSamplerState<>::GetRHI());
+				Parameters->SvPositionToInputTextureUV = (
+					FScreenTransform::ChangeTextureBasisFromTo(FScreenPassTextureViewport(Output), FScreenTransform::ETextureBasis::TexelPosition, FScreenTransform::ETextureBasis::ViewportUV) *
+					FScreenTransform::ChangeTextureBasisFromTo(FScreenPassTextureViewport(Input), FScreenTransform::ETextureBasis::ViewportUV, FScreenTransform::ETextureBasis::TextureUV));
 				Parameters->bIsForward = bIsForward;
 				Parameters->Gamma = InGamma;
-				Parameters->InvGamma = 1.0f / InGamma;
+				Parameters->InvGamma = FMath::SafeDivide(1.0f, InGamma);
 				Parameters->bIsForward = bIsForward;
 				Parameters->RenderTargets[0] = Output.GetRenderTargetBinding();
 
 				TShaderMapRef<FCompositeCoreDisplayTransformShader> PixelShader(GlobalShaderMap);
-				AddDrawScreenPass(
+				FPixelShaderUtils::AddFullscreenPass(
 					GraphBuilder,
-					RDG_EVENT_NAME("Composite.DisplayTransform (%dx%d)", Viewport.Extent.X, Viewport.Extent.Y),
-					ViewInfo,
-					Viewport,
-					Viewport,
+					GlobalShaderMap,
+					RDG_EVENT_NAME("Composite.DisplayTransform (%dx%d)", Output.ViewRect.Width(), Output.ViewRect.Height()),
 					PixelShader,
-					Parameters
+					Parameters,
+					Output.ViewRect
 				);
 
 				return Output;
diff --git a/Engine/Plugins/Compositing/CompositeCore/Source/CompositeCore/Private/Passes/CompositeCorePassProxy.cpp b/Engine/Plugins/Compositing/CompositeCore/Source/CompositeCore/Private/Passes/CompositeCorePassProxy.cpp
--- a/Engine/Plugins/Compositing/CompositeCore/Source/CompositeCore/Private/Passes/CompositeCorePassProxy.cpp
+++ b/Engine/Plugins/Compositing/CompositeCore/Source/CompositeCore/Private/Passes/CompositeCorePassProxy.cpp
@@ -105,7 +105,8 @@ FScreenPassRenderTarget FCompositeCorePassProxy::CreateOutputRenderTarget(FRDGBu
 	OutputDesc.Flags	= TexCreate_RenderTargetable | TexCreate_ShaderResource | TexCreate_UAV;
 	OutputDesc.Extent	= OutputViewRect.Size();
 
-	return FScreenPassRenderTarget(GraphBuilder.CreateTexture(OutputDesc, InName), OutputViewRect, InView.GetOverwriteLoadAction());
+	const FIntRect ZeroBasedViewRect(FIntPoint::ZeroValue, OutputViewRect.Size());
+	return FScreenPassRenderTarget(GraphBuilder.CreateTexture(OutputDesc, InName), ZeroBasedViewRect, InView.GetOverwriteLoadAction());
 }
 
 bool FCompositeCorePassProxy::ValidateInputs(const UE::CompositeCore::FPassInputArray& Inputs) const
