Composure: Preserve input texture sampler filter through compositing passes.

Instead of hardcoding SF_Bilinear sampler states everywhere, the branch propagates the source texture's sampler filter (point vs bilinear) through the compositing pipeline.

#Synced-CL 51541912

diff --git a/Engine/Plugins/Compositing/Composite/Source/Composite/Private/Layers/CompositeLayerPlate.cpp b/Engine/Plugins/Compositing/Composite/Source/Composite/Private/Layers/CompositeLayerPlate.cpp
--- a/Engine/Plugins/Compositing/Composite/Source/Composite/Private/Layers/CompositeLayerPlate.cpp
+++ b/Engine/Plugins/Compositing/Composite/Source/Composite/Private/Layers/CompositeLayerPlate.cpp
@@ -514,6 +514,10 @@ UE::CompositeCore::ResourceId UCompositeLayerPlate::AddPreprocessingPasses(
 				FResourceMetadata Metadata = {};
 				Metadata.bDistorted = true;
 				Metadata.DebugName = TEXT("CompositePlateTex");
+				if (IsValid(Texture))
+				{
+					Metadata.Filter = (Texture->Filter == TF_Nearest) ? SF_Point : SF_Bilinear;
+				}
 
 				OutputTextureId = InContext.FindOrCreateExternalTexture(GetRenderTargetFn(), Metadata);
 
@@ -541,6 +545,7 @@ bool UCompositeLayerPlate::GetProxy(const UE::CompositeCore::FPassInputDecl& /*I
 		FResourceMetadata Metadata = {};
 		Metadata.bDistorted = true;
 		Metadata.DebugName = TEXT("CompositePlateTex");
+		Metadata.Filter = (Texture->Filter == TF_Nearest) ? SF_Point : SF_Bilinear;
 
 		TextureId = InContext.FindOrCreateExternalTexture(Texture, Metadata);
 		
diff --git a/Engine/Plugins/Compositing/Composite/Source/Composite/Private/Passes/CompositePassAlphaTransform.cpp b/Engine/Plugins/Compositing/Composite/Source/Composite/Private/Passes/CompositePassAlphaTransform.cpp
--- a/Engine/Plugins/Compositing/Composite/Source/Composite/Private/Passes/CompositePassAlphaTransform.cpp
+++ b/Engine/Plugins/Compositing/Composite/Source/Composite/Private/Passes/CompositePassAlphaTransform.cpp
@@ -49,10 +49,10 @@ namespace UE
 					Output = CreateOutputRenderTarget(GraphBuilder, InView, PassContext.OutputViewRect, InputTexture.Texture->Desc, TEXT("AlphaTransformCompositePass"));
 				}
 
-				FRHISamplerState* BilinearClampSampler = TStaticSamplerState<SF_Bilinear>::GetRHI();
+				FRHISamplerState* InputSampler = Metadata.GetSamplerState();
 				FGlobalShaderMap* GlobalShaderMap = GetGlobalShaderMap(InView.GetFeatureLevel());
 				FCompositePassAlphaTransformShader::FParameters* Parameters = GraphBuilder.AllocParameters<FCompositePassAlphaTransformShader::FParameters>();
-				Parameters->Input = GetScreenPassTextureInput(InputTexture, BilinearClampSampler);
+				Parameters->Input = GetScreenPassTextureInput(InputTexture, InputSampler);
 				Parameters->SvPositionToInputTextureUV = (
 					FScreenTransform::ChangeTextureBasisFromTo(FScreenPassTextureViewport(Output), FScreenTransform::ETextureBasis::TexelPosition, FScreenTransform::ETextureBasis::ViewportUV) *
 					FScreenTransform::ChangeTextureBasisFromTo(FScreenPassTextureViewport(InputTexture), FScreenTransform::ETextureBasis::ViewportUV, FScreenTransform::ETextureBasis::TextureUV));
diff --git a/Engine/Plugins/Compositing/Composite/Source/Composite/Private/Passes/CompositePassCenteredScale.cpp b/Engine/Plugins/Compositing/Composite/Source/Composite/Private/Passes/CompositePassCenteredScale.cpp
--- a/Engine/Plugins/Compositing/Composite/Source/Composite/Private/Passes/CompositePassCenteredScale.cpp
+++ b/Engine/Plugins/Compositing/Composite/Source/Composite/Private/Passes/CompositePassCenteredScale.cpp
@@ -58,7 +58,8 @@ namespace UE
 
 				FGlobalShaderMap* GlobalShaderMap = GetGlobalShaderMap(InView.GetFeatureLevel());
 				FCompositeCenteredScaleShader::FParameters* Parameters = GraphBuilder.AllocParameters<FCompositeCenteredScaleShader::FParameters>();
-				Parameters->Input = GetScreenPassTextureInput(Input, TStaticSamplerState<SF_Bilinear>::GetRHI());
+				FRHISamplerState* InputSampler = Metadata.GetSamplerState();
+				Parameters->Input = GetScreenPassTextureInput(Input, InputSampler);
 				Parameters->SvPositionToOutputViewportUV = FScreenTransform::ChangeTextureBasisFromTo(
 					FScreenPassTextureViewport(Output),
 					FScreenTransform::ETextureBasis::TexelPosition,
diff --git a/Engine/Plugins/Compositing/Composite/Source/Composite/Private/Passes/CompositePassColorGrading.cpp b/Engine/Plugins/Compositing/Composite/Source/Composite/Private/Passes/CompositePassColorGrading.cpp
--- a/Engine/Plugins/Compositing/Composite/Source/Composite/Private/Passes/CompositePassColorGrading.cpp
+++ b/Engine/Plugins/Compositing/Composite/Source/Composite/Private/Passes/CompositePassColorGrading.cpp
@@ -80,7 +80,8 @@ namespace UE::Composite::Private
 
 		FGlobalShaderMap* GlobalShaderMap = GetGlobalShaderMap(InView.GetFeatureLevel());
 		FCompositeColorGradingShader::FParameters* Parameters = GraphBuilder.AllocParameters<FCompositeColorGradingShader::FParameters>();
-		Parameters->Input = GetScreenPassTextureInput(Input, TStaticSamplerState<SF_Bilinear>::GetRHI());
+		FRHISamplerState* InputSampler = Metadata.GetSamplerState();
+		Parameters->Input = GetScreenPassTextureInput(Input, InputSampler);
 		Parameters->WorkingColorSpace = GDefaultWorkingColorSpaceUniformBuffer.GetUniformBufferRef();
 		Parameters->Global = FCompositeColorGradingShader::ToParameters(ColorGradingSettings.Global);
 		Parameters->Shadows = FCompositeColorGradingShader::ToParameters(ColorGradingSettings.Shadows);
diff --git a/Engine/Plugins/Compositing/Composite/Source/Composite/Private/Passes/CompositePassColorKeyer.cpp b/Engine/Plugins/Compositing/Composite/Source/Composite/Private/Passes/CompositePassColorKeyer.cpp
--- a/Engine/Plugins/Compositing/Composite/Source/Composite/Private/Passes/CompositePassColorKeyer.cpp
+++ b/Engine/Plugins/Compositing/Composite/Source/Composite/Private/Passes/CompositePassColorKeyer.cpp
@@ -96,6 +96,7 @@ namespace UE
 				FScreenPassRenderTarget Output = FScreenPassRenderTarget(GraphBuilder.CreateTexture(OutputDesc, TEXT("CompositeTextureDenoised")), FIntRect(FIntPoint::ZeroValue, TextureSize), ERenderTargetLoadAction::ENoAction);
 
 				FCompositePassDenoiseShader::FParameters* PassParameters = GraphBuilder.AllocParameters<FCompositePassDenoiseShader::FParameters>();
+				// Always use point sampling: denoising needs exact texel values for its spatial filter.
 				PassParameters->Input = GetScreenPassTextureInput(Input, TStaticSamplerState<SF_Point>::GetRHI());
 				PassParameters->RWOutputTexture = GraphBuilder.CreateUAV(Output.Texture);
 
@@ -135,11 +136,12 @@ namespace UE
 
 				FGlobalShaderMap* GlobalShaderMap = GetGlobalShaderMap(InView.GetFeatureLevel());
 
-				const bool bSameSize = Output.ViewRect.Size() == Input.ViewRect.Size();
 				const FScreenPassTextureViewportParameters OutputParameters = GetScreenPassTextureViewportParameters(FScreenPassTextureViewport(Output));
 
+				FRHISamplerState* InputSampler = Metadata.GetSamplerState();
+
 				FCompositePassColorKeyerShader::FParameters* PassParameters = GraphBuilder.AllocParameters<FCompositePassColorKeyerShader::FParameters>();
-				PassParameters->Input = GetScreenPassTextureInput(Input, bSameSize ? TStaticSamplerState<SF_Point>::GetRHI() : TStaticSamplerState<SF_Bilinear>::GetRHI());
+				PassParameters->Input = GetScreenPassTextureInput(Input, InputSampler);
 				PassParameters->RWOutputTexture = GraphBuilder.CreateUAV(Output.Texture);
 				PassParameters->Output = OutputParameters;
 				PassParameters->SvPositionToInputTextureUV = (
@@ -155,6 +157,7 @@ namespace UE
 	
 				FRDGTextureRef CleanPlateResource = nullptr;
 				PassParameters->CleanPlate.Texture = GSystemTextures.GetBlackDummy(GraphBuilder);
+				// Clean plate is a separate reference texture not driven by the input metadata; bilinear is always appropriate.
 				PassParameters->CleanPlate.Sampler = TStaticSamplerState<SF_Bilinear>::GetRHI();
 
 				TStrongObjectPtr<UTexture> CleanPlateTexturePtr = CleanPlateWeakPtr.Pin();
diff --git a/Engine/Plugins/Compositing/Composite/Source/Composite/Private/Passes/CompositePassDistortion.cpp b/Engine/Plugins/Compositing/Composite/Source/Composite/Private/Passes/CompositePassDistortion.cpp
--- a/Engine/Plugins/Compositing/Composite/Source/Composite/Private/Passes/CompositePassDistortion.cpp
+++ b/Engine/Plugins/Compositing/Composite/Source/Composite/Private/Passes/CompositePassDistortion.cpp
@@ -94,6 +94,7 @@ public:
 			FScreenPassTextureViewport Viewport = FScreenPassTextureViewport(Output);
 
 			FCompositeDistortionShader::FParameters* PassParameters = GraphBuilder.AllocParameters<FCompositeDistortionShader::FParameters>();
+			// Always use bilinear: distortion resamples through a displacement LUT, mapping texels to fractional coordinates.
 			PassParameters->Input = GetScreenPassTextureInput(Input, TStaticSamplerState<SF_Bilinear>::GetRHI());
 			PassParameters->Output = GetScreenPassTextureViewportParameters(FScreenPassTextureViewport(Output));
 			if (ViewDistortionLUT.IsEnabled())
diff --git a/Engine/Plugins/Compositing/Composite/Source/Composite/Private/Passes/CompositePassMasking.cpp b/Engine/Plugins/Compositing/Composite/Source/Composite/Private/Passes/CompositePassMasking.cpp
--- a/Engine/Plugins/Compositing/Composite/Source/Composite/Private/Passes/CompositePassMasking.cpp
+++ b/Engine/Plugins/Compositing/Composite/Source/Composite/Private/Passes/CompositePassMasking.cpp
@@ -55,7 +55,7 @@ namespace UE
 			{
 				FCompositePassMaskingShader::FInputPassTextureParameters Parameters;
 				Parameters.Texture = Input.Texture.Texture;
-				Parameters.Sampler = TStaticSamplerState<SF_Bilinear>::GetRHI();
+				Parameters.Sampler = Input.Metadata.GetSamplerState();
 				Parameters.bInvertedAlpha = static_cast<uint32>(Input.Metadata.bInvertedAlpha);
 				Parameters.SourceEncoding = static_cast<uint32>(Input.Metadata.Encoding);
 
@@ -136,7 +136,10 @@ bool UCompositePassMasking::GetProxy(const UE::CompositeCore::FPassInputDecl& In
 
 	if (ensureMsgf(bValidTexture, TEXT("Texture is expected to be valid for active pass.")))
 	{
-		ResourceId TextureId = InContext.FindOrCreateExternalTexture(MaskTexture, { .DebugName = TEXT("MaskTexture") });
+		UE::CompositeCore::FResourceMetadata MaskMetadata = { .DebugName = TEXT("MaskTexture") };
+		MaskMetadata.Filter = (MaskTexture->Filter == TF_Nearest) ? SF_Point : SF_Bilinear;
+
+		ResourceId TextureId = InContext.FindOrCreateExternalTexture(MaskTexture, MaskMetadata);
 
 		FPassInputDecl MaskingTextureInput;
 		MaskingTextureInput.Set<FPassExternalResourceDesc>({ TextureId });
diff --git a/Engine/Plugins/Compositing/CompositeCore/Source/CompositeCore/Private/Passes/CompositeCorePassMergeProxy.cpp b/Engine/Plugins/Compositing/CompositeCore/Source/CompositeCore/Private/Passes/CompositeCorePassMergeProxy.cpp
--- a/Engine/Plugins/Compositing/CompositeCore/Source/CompositeCore/Private/Passes/CompositeCorePassMergeProxy.cpp
+++ b/Engine/Plugins/Compositing/CompositeCore/Source/CompositeCore/Private/Passes/CompositeCorePassMergeProxy.cpp
@@ -84,7 +84,7 @@ namespace UE
 			FCompositeCoreMergeShader::FInputPassTextureParameters Parameters;
 
 			Parameters.Texture = Input.Texture.Texture;
-			Parameters.Sampler = TStaticSamplerState<SF_Bilinear>::GetRHI();
+			Parameters.Sampler = Input.Metadata.GetSamplerState();
 			Parameters.bInvertedAlpha = Input.Metadata.bInvertedAlpha;
 			Parameters.SourceEncoding = static_cast<uint32>(Input.Metadata.Encoding);
 			Parameters.Exposure = 1.0f;
diff --git a/Engine/Plugins/Compositing/CompositeCore/Source/CompositeCore/Private/Passes/CompositeCorePassProxy.cpp b/Engine/Plugins/Compositing/CompositeCore/Source/CompositeCore/Private/Passes/CompositeCorePassProxy.cpp
--- a/Engine/Plugins/Compositing/CompositeCore/Source/CompositeCore/Private/Passes/CompositeCorePassProxy.cpp
+++ b/Engine/Plugins/Compositing/CompositeCore/Source/CompositeCore/Private/Passes/CompositeCorePassProxy.cpp
@@ -8,6 +8,7 @@
 #include "FXRenderingUtils.h"
 #include "RenderGraphBuilder.h"
 #include "RenderGraphUtils.h"
+#include "RHIDefinitions.h"
 #include "PostProcess/LensDistortion.h"
 #include "PostProcess/PostProcessMaterialInputs.h"
 
@@ -91,6 +92,13 @@ namespace UE
 	}
 }
 
+FRHISamplerState* UE::CompositeCore::FResourceMetadata::GetSamplerState() const
+{
+	return (Filter == SF_Point)
+		? TStaticSamplerState<SF_Point>::GetRHI()
+		: TStaticSamplerState<SF_Bilinear>::GetRHI();
+}
+
 FCompositeCorePassProxy::FCompositeCorePassProxy(UE::CompositeCore::FPassInputDeclArray InPassDeclaredInputs)
 	: PassDeclaredInputs(MoveTemp(InPassDeclaredInputs))
 	, PassDeclaredOutputOverride()
diff --git a/Engine/Plugins/Compositing/CompositeCore/Source/CompositeCore/Public/Passes/CompositeCorePassProxy.h b/Engine/Plugins/Compositing/CompositeCore/Source/CompositeCore/Public/Passes/CompositeCorePassProxy.h
--- a/Engine/Plugins/Compositing/CompositeCore/Source/CompositeCore/Public/Passes/CompositeCorePassProxy.h
+++ b/Engine/Plugins/Compositing/CompositeCore/Source/CompositeCore/Public/Passes/CompositeCorePassProxy.h
@@ -21,6 +21,7 @@ struct FScreenPassTexture;
 struct FScreenPassRenderTarget;
 struct FPostProcessMaterialInputs;
 class FCompositeCorePassProxy;
+class FRHISamplerState;
 struct IPooledRenderTarget;
 class UTexture;
 class USceneCaptureComponent2D;
@@ -63,16 +64,23 @@ namespace UE
 			/** Source color encoding. */
 			EEncoding Encoding = EEncoding::Linear;
 
+			/** Sampler filter mode for the texture resource. */
+			ESamplerFilter Filter = SF_Bilinear;
+
+			/** Returns the RHI sampler state matching the Filter setting. */
+			COMPOSITECORE_API FRHISamplerState* GetSamplerState() const;
+
 			/** Debug name */
 			const TCHAR* DebugName = TEXT("CompositeTexture");
-			
+
 			/** Equality operator */
 			bool operator==(const FResourceMetadata& Other) const
 			{
 				return bInvertedAlpha == Other.bInvertedAlpha
 					&& bDistorted == Other.bDistorted
 					&& bPreExposed == Other.bPreExposed
-					&& Encoding == Other.Encoding;
+					&& Encoding == Other.Encoding
+					&& Filter == Other.Filter;
 
 				// DebugName is intentionally ignored
 			}
