Composure: Initial support for Composite layer & pass sequencer integration, such that their properties can be keyframed.

* Added IObjectSchema implementation registered with the Sequencer module for two-level sub-object discovery: Actor -> Layers, Layer -> Passes.
* Added "Composite" category to all properties for correct organization in Sequencer. (Customization update to no longer rely on "NAME_None".
* Added the Interp specifier to ~30 UPROPERTY declarations across layers and passes, making them visible and animatable in Sequencer.

#Synced-CL 51541912

diff --git a/Engine/Plugins/Compositing/Composite/Source/Composite/Public/Layers/CompositeLayerBase.h b/Engine/Plugins/Compositing/Composite/Source/Composite/Public/Layers/CompositeLayerBase.h
--- a/Engine/Plugins/Compositing/Composite/Source/Composite/Public/Layers/CompositeLayerBase.h
+++ b/Engine/Plugins/Compositing/Composite/Source/Composite/Public/Layers/CompositeLayerBase.h
@@ -80,7 +80,7 @@ public:
 #endif
 
 	/** Merge operation applied on Input0 with Input1. */
-	UPROPERTY(EditAnywhere, BlueprintReadWrite, Category = "", meta = (DisplayPriority = "3"))
+	UPROPERTY(EditAnywhere, BlueprintReadWrite, Interp, Category = "Composite", meta = (DisplayPriority = "3"))
 	ECompositeCoreMergeOp Operation = ECompositeCoreMergeOp::Over;
 
 	/** Convenience function to register post-processing child passes. */
diff --git a/Engine/Plugins/Compositing/Composite/Source/Composite/Public/Layers/CompositeLayerMainRender.h b/Engine/Plugins/Compositing/Composite/Source/Composite/Public/Layers/CompositeLayerMainRender.h
--- a/Engine/Plugins/Compositing/Composite/Source/Composite/Public/Layers/CompositeLayerMainRender.h
+++ b/Engine/Plugins/Compositing/Composite/Source/Composite/Public/Layers/CompositeLayerMainRender.h
@@ -27,7 +27,7 @@ public:
 	UE_API virtual bool GetProxy(const UE::CompositeCore::FPassInputDecl& InputDecl, FCompositeTraversalContext& InContext, FSceneRenderingBulkObjectAllocator& InFrameAllocator, FCompositeCorePassProxy*& OutProxy) const override;
 
 	/** Sub-passes applied onto the main render input before it is merged as this layer. */
-	UPROPERTY(EditAnywhere, BlueprintReadWrite, Instanced, Category = "", meta = (DisallowedClasses = "/Script/Composite.CompositeLayerBase, /Script/Composite.CompositePassDistortion"))
+	UPROPERTY(EditAnywhere, BlueprintReadWrite, Instanced, Category = "Composite", meta = (DisallowedClasses = "/Script/Composite.CompositeLayerBase, /Script/Composite.CompositePassDistortion"))
 	TArray<TObjectPtr<UCompositePassBase>> LayerPasses;
 };
 
diff --git a/Engine/Plugins/Compositing/Composite/Source/Composite/Public/Layers/CompositeLayerPlanarReflection.h b/Engine/Plugins/Compositing/Composite/Source/Composite/Public/Layers/CompositeLayerPlanarReflection.h
--- a/Engine/Plugins/Compositing/Composite/Source/Composite/Public/Layers/CompositeLayerPlanarReflection.h
+++ b/Engine/Plugins/Compositing/Composite/Source/Composite/Public/Layers/CompositeLayerPlanarReflection.h
@@ -78,7 +78,7 @@ public:
 	UE_API void SetReflectionStrength(float InReflectionStrength);
 
 	/** Sub-passes applied during post-processing. */
-	UPROPERTY(EditAnywhere, BlueprintReadWrite, Instanced, Category = "", meta = (DisallowedClasses = "/Script/Composite.CompositeLayerBase, /Script/Composite.CompositePassDistortion", DisplayAfter = "bCustomRenderPass"))
+	UPROPERTY(EditAnywhere, BlueprintReadWrite, Instanced, Category = "Composite", meta = (DisallowedClasses = "/Script/Composite.CompositeLayerBase, /Script/Composite.CompositePassDistortion", DisplayAfter = "bCustomRenderPass"))
 	TArray<TObjectPtr<UCompositePassBase>> LayerPasses;
 
 private:
@@ -90,23 +90,23 @@ private:
 	void UpdatePlanarReflectionTransform(USceneCaptureComponent2D* CaptureComponent);
 
 	/** List of actors to be rendered from the planar reflection. */
-	UPROPERTY(EditAnywhere, BlueprintGetter = GetActors, BlueprintSetter = SetActors, Category = "", meta = (AllowPrivateAccess = true))
+	UPROPERTY(EditAnywhere, BlueprintGetter = GetActors, BlueprintSetter = SetActors, Category = "Composite", meta = (AllowPrivateAccess = true))
 	TArray<TSoftObjectPtr<AActor>> Actors;
 
 	/** Render target resolution for the planar reflection. */
-	UPROPERTY(EditAnywhere, BlueprintGetter = GetRenderTargetResolution, BlueprintSetter = SetRenderTargetResolution, Interp, Category = "", meta = (AllowPrivateAccess = true))
+	UPROPERTY(EditAnywhere, BlueprintGetter = GetRenderTargetResolution, BlueprintSetter = SetRenderTargetResolution, Interp, Category = "Composite", meta = (AllowPrivateAccess = true))
 	FIntPoint RenderTargetResolution;
 
 	/** Strength of the planar reflection. */
-	UPROPERTY(EditAnywhere, BlueprintGetter = GetReflectionStrength, BlueprintSetter = SetReflectionStrength, Category = "", meta = (AllowPrivateAccess = true, ClampMin = "0.0", ClampMax = "1.0", UIMin = "0.0", UIMax = "1.0"))
+	UPROPERTY(EditAnywhere, BlueprintGetter = GetReflectionStrength, BlueprintSetter = SetReflectionStrength, Interp, Category = "Composite", meta = (AllowPrivateAccess = true, ClampMin = "0.0", ClampMax = "1.0", UIMin = "0.0", UIMax = "1.0"))
 	float ReflectionStrength = 0.5f;
 
 	/** World-space Z coordinate of the horizontal plane for planar reflection. Note that we currently do not clip against this plane. */
-	UPROPERTY(EditAnywhere, BlueprintReadWrite, Category = "", meta = (AllowPrivateAccess = true))
+	UPROPERTY(EditAnywhere, BlueprintReadWrite, Category = "Composite", meta = (AllowPrivateAccess = true))
 	float PlaneHeight = 0.0f;
 
 	/** Convenience toggle to render the planar reflection as a fast/minimal Custom Render Pass, inlined in the main render & without support for lighting. */
-	UPROPERTY(EditAnywhere, BlueprintGetter = IsCustomRenderPass, BlueprintSetter = SetCustomRenderPass, Category = "", AdvancedDisplay, meta = (AllowPrivateAccess = true))
+	UPROPERTY(EditAnywhere, BlueprintGetter = IsCustomRenderPass, BlueprintSetter = SetCustomRenderPass, Category = "Composite", AdvancedDisplay, meta = (AllowPrivateAccess = true))
 	bool bCustomRenderPass;
 };
 
diff --git a/Engine/Plugins/Compositing/Composite/Source/Composite/Public/Layers/CompositeLayerPlate.h b/Engine/Plugins/Compositing/Composite/Source/Composite/Public/Layers/CompositeLayerPlate.h
--- a/Engine/Plugins/Compositing/Composite/Source/Composite/Public/Layers/CompositeLayerPlate.h
+++ b/Engine/Plugins/Compositing/Composite/Source/Composite/Public/Layers/CompositeLayerPlate.h
@@ -119,12 +119,12 @@ public:
 
 private:
 	/** Actors whose primitive mesh components will be marked as primitive alpha holdout, and rendered separately with the projected (media) texture. */
-	UPROPERTY(EditAnywhere, BlueprintGetter = GetCompositeMeshes, BlueprintSetter = SetCompositeMeshes, Category = "", meta = (AllowPrivateAccess = true))
+	UPROPERTY(EditAnywhere, BlueprintGetter = GetCompositeMeshes, BlueprintSetter = SetCompositeMeshes, Category = "Composite", meta = (AllowPrivateAccess = true))
 	TArray<TSoftObjectPtr<AActor>> CompositeMeshes;
 
 public:
 	/** Media texture for compositing layer. */
-	UPROPERTY(EditAnywhere, BlueprintReadWrite, Category = "", meta = (AllowedClasses="/Script/Engine.Texture2D, /Script/MediaAssets.MediaTexture"))
+	UPROPERTY(EditAnywhere, BlueprintReadWrite, Category = "Composite", meta = (AllowedClasses="/Script/Engine.Texture2D, /Script/MediaAssets.MediaTexture"))
 	TObjectPtr<UTexture> Texture;
 
 	/**
@@ -132,7 +132,7 @@ public:
 	* 
 	* The texture is replaced with an internal render target.
 	*/
-	UPROPERTY(EditAnywhere, BlueprintReadWrite, Instanced, Category = "", meta = (DisallowedClasses = "/Script/Composite.CompositeLayerBase, /Script/Composite.CompositePassMaterial, /Script/Composite.CompositePassDistortion, /Script/Composite.CompositePassFXAA, /Script/Composite.CompositePassSMAA"))
+	UPROPERTY(EditAnywhere, BlueprintReadWrite, Instanced, Category = "Composite", meta = (DisallowedClasses = "/Script/Composite.CompositeLayerBase, /Script/Composite.CompositePassMaterial, /Script/Composite.CompositePassDistortion, /Script/Composite.CompositePassFXAA, /Script/Composite.CompositePassSMAA"))
 	TArray<TObjectPtr<UCompositePassBase>> MediaPasses;
 
 	/**
@@ -144,11 +144,11 @@ public:
 	* We now have an automatic scene pass to undistort a plate projected onto composite meshes. Using Texture mode can therefore
 	* avoid the automatic undistort -> distort pipeline.
 	*/
-	UPROPERTY(BlueprintReadWrite, Instanced, Category = "", meta = (AllowedClasses = "/Script/Composite.CompositePassDistortion", DisallowedClasses = "/Script/Composite.CompositeLayerBase"))
+	UPROPERTY(BlueprintReadWrite, Instanced, Category = "Composite", meta = (AllowedClasses = "/Script/Composite.CompositePassDistortion", DisallowedClasses = "/Script/Composite.CompositeLayerBase"))
 	TArray<TObjectPtr<UCompositePassBase>> ScenePasses;
 
 	/** Passes applied during final compositing in post-processing. */
-	UPROPERTY(EditAnywhere, BlueprintReadWrite, Instanced, Category = "", meta = (DisallowedClasses = "/Script/Composite.CompositeLayerBase, /Script/Composite.CompositePassDistortion"))
+	UPROPERTY(EditAnywhere, BlueprintReadWrite, Instanced, Category = "Composite", meta = (DisallowedClasses = "/Script/Composite.CompositeLayerBase, /Script/Composite.CompositePassDistortion"))
 	TArray<TObjectPtr<UCompositePassBase>> LayerPasses;
 
 	/** Composite.Plate.ControlMeshVisibility console variable value. */
@@ -156,7 +156,7 @@ public:
 
 private:
 	/** The plate sampling & rendering mode: direct texture sampling or separate composite mesh render. */
-	UPROPERTY(EditAnywhere, BlueprintGetter = GetPlateMode, BlueprintSetter = SetPlateMode, Category = "", AdvancedDisplay, meta = (AllowPrivateAccess = true, DisplayName = "Mode"))
+	UPROPERTY(EditAnywhere, BlueprintGetter = GetPlateMode, BlueprintSetter = SetPlateMode, Interp, Category = "Composite", AdvancedDisplay, meta = (AllowPrivateAccess = true, DisplayName = "Mode"))
 	ECompositePlateMode PlateMode;
 
 	/**
@@ -166,7 +166,7 @@ private:
 	* artifacts and the tone curve. However when working in 3D, holdout can be disabled to only preserve the
 	* plate projection functionality and media passes.
 	*/
-	UPROPERTY(EditAnywhere, BlueprintGetter = GetIsHoldoutEnabled, BlueprintSetter = SetIsHoldoutEnabled, Category = "", AdvancedDisplay, meta = (AllowPrivateAccess = true))
+	UPROPERTY(EditAnywhere, BlueprintGetter = GetIsHoldoutEnabled, BlueprintSetter = SetIsHoldoutEnabled, Interp, Category = "Composite", AdvancedDisplay, meta = (AllowPrivateAccess = true))
 	bool bIsHoldoutEnabled;
 
 private:
diff --git a/Engine/Plugins/Compositing/Composite/Source/Composite/Public/Layers/CompositeLayerProcessing.h b/Engine/Plugins/Compositing/Composite/Source/Composite/Public/Layers/CompositeLayerProcessing.h
--- a/Engine/Plugins/Compositing/Composite/Source/Composite/Public/Layers/CompositeLayerProcessing.h
+++ b/Engine/Plugins/Compositing/Composite/Source/Composite/Public/Layers/CompositeLayerProcessing.h
@@ -32,7 +32,7 @@ public:
 #endif
 
 	/** Sub-passes applied onto the previous layer. */
-	UPROPERTY(EditAnywhere, BlueprintReadWrite, Instanced, Category = "", meta = (DisallowedClasses = "/Script/Composite.CompositeLayerBase, /Script/Composite.CompositePassDistortion"))
+	UPROPERTY(EditAnywhere, BlueprintReadWrite, Instanced, Category = "Composite", meta = (DisallowedClasses = "/Script/Composite.CompositeLayerBase, /Script/Composite.CompositePassDistortion"))
 	TArray<TObjectPtr<UCompositePassBase>> LayerPasses;
 };
 
diff --git a/Engine/Plugins/Compositing/Composite/Source/Composite/Public/Layers/CompositeLayerSceneCapture.h b/Engine/Plugins/Compositing/Composite/Source/Composite/Public/Layers/CompositeLayerSceneCapture.h
--- a/Engine/Plugins/Compositing/Composite/Source/Composite/Public/Layers/CompositeLayerSceneCapture.h
+++ b/Engine/Plugins/Compositing/Composite/Source/Composite/Public/Layers/CompositeLayerSceneCapture.h
@@ -80,7 +80,7 @@ public:
 	UE_API void SetRenderTargetResolution(FIntPoint InRenderTargetResolution);
 
 	/** Sub-passes applied during post-processing. */
-	UPROPERTY(EditAnywhere, BlueprintReadWrite, Instanced, Category = "", meta = (DisallowedClasses = "/Script/Composite.CompositeLayerBase, /Script/Composite.CompositePassDistortion", DisplayAfter = "bCustomRenderPass"))
+	UPROPERTY(EditAnywhere, BlueprintReadWrite, Instanced, Category = "Composite", meta = (DisallowedClasses = "/Script/Composite.CompositeLayerBase, /Script/Composite.CompositePassDistortion", DisplayAfter = "bCustomRenderPass"))
 	TArray<TObjectPtr<UCompositePassBase>> LayerPasses;
 
 private:
@@ -99,7 +99,7 @@ private:
 	 * 
 	 * Automatically sets the scene capture's ShowOnlyComponents when updated.
 	*/
-	UPROPERTY(EditAnywhere, BlueprintGetter = GetActors, BlueprintSetter = SetActors, Category = "", meta = (AllowPrivateAccess = true))
+	UPROPERTY(EditAnywhere, BlueprintGetter = GetActors, BlueprintSetter = SetActors, Category = "Composite", meta = (AllowPrivateAccess = true))
 	TArray<TSoftObjectPtr<AActor>> Actors;
 
 #if WITH_EDITORONLY_DATA
@@ -109,11 +109,11 @@ private:
 #endif
 
 	/** Render target resolution for the scene capture. */
-	UPROPERTY(EditAnywhere, BlueprintGetter = GetRenderTargetResolution, BlueprintSetter = SetRenderTargetResolution, Interp, Category = "", meta = (AllowPrivateAccess = true))
+	UPROPERTY(EditAnywhere, BlueprintGetter = GetRenderTargetResolution, BlueprintSetter = SetRenderTargetResolution, Interp, Category = "Composite", meta = (AllowPrivateAccess = true))
 	FIntPoint RenderTargetResolution;
 
 	/** Visibility setting applied to registered actors. By default, registered meshes will only be visible in scene captures and no longer in the main render. */
-	UPROPERTY(EditAnywhere, BlueprintGetter = IsVisibleInSceneCaptureOnly, BlueprintSetter = SetVisibleInSceneCaptureOnly, Category = "", meta = (AllowPrivateAccess = true))
+	UPROPERTY(EditAnywhere, BlueprintGetter = IsVisibleInSceneCaptureOnly, BlueprintSetter = SetVisibleInSceneCaptureOnly, Category = "Composite", meta = (AllowPrivateAccess = true))
 	bool bVisibleInSceneCaptureOnly;
 
 	/**
@@ -121,7 +121,7 @@ private:
 	 * 
 	 * Automatically sets the scene capture's CaptureSource, bRenderInMainRenderer & bIgnoreScreenPercentage when updated.
 	 */
-	UPROPERTY(EditAnywhere, BlueprintGetter = IsCustomRenderPass, BlueprintSetter = SetCustomRenderPass, Category = "", meta = (AllowPrivateAccess = true))
+	UPROPERTY(EditAnywhere, BlueprintGetter = IsCustomRenderPass, BlueprintSetter = SetCustomRenderPass, Category = "Composite", meta = (AllowPrivateAccess = true))
 	bool bCustomRenderPass;
 
 	friend class FCompositeLayerSceneCaptureCustomization;
diff --git a/Engine/Plugins/Compositing/Composite/Source/Composite/Public/Layers/CompositeLayerShadowReflection.h b/Engine/Plugins/Compositing/Composite/Source/Composite/Public/Layers/CompositeLayerShadowReflection.h
--- a/Engine/Plugins/Compositing/Composite/Source/Composite/Public/Layers/CompositeLayerShadowReflection.h
+++ b/Engine/Plugins/Compositing/Composite/Source/Composite/Public/Layers/CompositeLayerShadowReflection.h
@@ -120,11 +120,11 @@ private:
 
 private:
 	/** Render target resolution for the scene captures. */
-	UPROPERTY(EditAnywhere, BlueprintGetter = GetRenderTargetResolution, BlueprintSetter = SetRenderTargetResolution, Interp, Category = "", meta = (AllowPrivateAccess = true))
+	UPROPERTY(EditAnywhere, BlueprintGetter = GetRenderTargetResolution, BlueprintSetter = SetRenderTargetResolution, Interp, Category = "Composite", meta = (AllowPrivateAccess = true))
 	FIntPoint RenderTargetResolution;
 
 	/** List of reflection & shadow caster actors. */
-	UPROPERTY(EditAnywhere, BlueprintGetter = GetActors, BlueprintSetter = SetActors, Category = "")
+	UPROPERTY(EditAnywhere, BlueprintGetter = GetActors, BlueprintSetter = SetActors, Category = "Composite")
 	TArray<TSoftObjectPtr<AActor>> Actors;
 
 #if WITH_EDITORONLY_DATA
diff --git a/Engine/Plugins/Compositing/Composite/Source/Composite/Public/Layers/CompositeLayerSingleLightShadow.h b/Engine/Plugins/Compositing/Composite/Source/Composite/Public/Layers/CompositeLayerSingleLightShadow.h
--- a/Engine/Plugins/Compositing/Composite/Source/Composite/Public/Layers/CompositeLayerSingleLightShadow.h
+++ b/Engine/Plugins/Compositing/Composite/Source/Composite/Public/Layers/CompositeLayerSingleLightShadow.h
@@ -63,7 +63,7 @@ public:
 	* 
 	* Only directional lights are currently supported.
 	*/
-	UPROPERTY(EditAnywhere, BlueprintReadWrite, Category = "", meta=(AllowedClasses="/Script/Engine.DirectionalLight"))
+	UPROPERTY(EditAnywhere, BlueprintReadWrite, Category = "Composite", meta=(AllowedClasses="/Script/Engine.DirectionalLight"))
 	TObjectPtr<ALight> Light;
 
 	/** Get shadow caster actors. */
@@ -78,14 +78,14 @@ public:
 	* When enable, the shadow map render transform and orthographic width are automatically derived from the selected meshes and light angle.
 	* When disabled, the shadow map transform is copied from the light transform, which should be piloted to align with shadow casting geometry.
 	*/
-	UPROPERTY(EditAnywhere, BlueprintReadWrite, Category = "")
+	UPROPERTY(EditAnywhere, BlueprintReadWrite, Category = "Composite")
 	bool bAutomaticBounds;
 
 	/**
 	* The desired width (in world units) of the shadow map view (ignored if light is not directional).
 	* Reduce size to only contain the shadow casting geometry in view of the light.
 	*/
-	UPROPERTY(EditAnywhere, BlueprintReadWrite, Category = "", meta = (EditCondition = "!bAutomaticBounds", EditConditionHides))
+	UPROPERTY(EditAnywhere, BlueprintReadWrite, Category = "Composite", meta = (EditCondition = "!bAutomaticBounds", EditConditionHides))
 	float OrthographicWidth;
 
 	/** Get the render target resolution for the scene depth capture. */
@@ -99,15 +99,15 @@ public:
 	/**
 	* Resolution of the shadow map texture. Lower resolutions will cause more blurry shadows.
 	*/
-	UPROPERTY(EditAnywhere, BlueprintReadWrite, Category = "", meta = (ClampMin = "128", UIMin = "128", UIMax="4096"))
+	UPROPERTY(EditAnywhere, BlueprintReadWrite, Category = "Composite", meta = (ClampMin = "128", UIMin = "128", UIMax="4096"))
 	int32 ShadowMapResolution;
 
 	/** Basic shadow strength multiplier, from 0 to 1. */
-	UPROPERTY(EditAnywhere, BlueprintReadWrite, Category = "", meta = (ClampMin = "0.0", ClampMax = "1.0", UIMin = "0.0", UIMax = "1.0"))
+	UPROPERTY(EditAnywhere, BlueprintReadWrite, Category = "Composite", meta = (ClampMin = "0.0", ClampMax = "1.0", UIMin = "0.0", UIMax = "1.0"))
 	float ShadowStrength;
 
 	/** Color grading pass settings. */
-	UPROPERTY(BlueprintReadWrite, EditAnywhere, Category = "")
+	UPROPERTY(BlueprintReadWrite, EditAnywhere, Category = "Composite")
 	FCompositeColorGradingPassSettings ColorGrading;
 
 #if WITH_EDITORONLY_DATA
@@ -128,11 +128,11 @@ private:
 
 private:
 	/** Render target resolution for the scene depth capture. */
-	UPROPERTY(EditAnywhere, BlueprintGetter = GetRenderTargetResolution, BlueprintSetter = SetRenderTargetResolution, Interp, Category = "", meta = (AllowPrivateAccess = true))
+	UPROPERTY(EditAnywhere, BlueprintGetter = GetRenderTargetResolution, BlueprintSetter = SetRenderTargetResolution, Interp, Category = "Composite", meta = (AllowPrivateAccess = true))
 	FIntPoint RenderTargetResolution;
 
 	/** List of shadow casting actors (required). */
-	UPROPERTY(EditAnywhere, BlueprintGetter = GetShadowCastingActors, BlueprintSetter = SetShadowCastingActors, Category = "", meta = (AllowPrivateAccess = true))
+	UPROPERTY(EditAnywhere, BlueprintGetter = GetShadowCastingActors, BlueprintSetter = SetShadowCastingActors, Category = "Composite", meta = (AllowPrivateAccess = true))
 	TArray<TSoftObjectPtr<AActor>> ShadowCastingActors;
 
 	/** Cached pointer to the composite-managed shadow map scene capture. */
diff --git a/Engine/Plugins/Compositing/Composite/Source/Composite/Public/Passes/CompositePassAlphaTransform.h b/Engine/Plugins/Compositing/Composite/Source/Composite/Public/Passes/CompositePassAlphaTransform.h
--- a/Engine/Plugins/Compositing/Composite/Source/Composite/Public/Passes/CompositePassAlphaTransform.h
+++ b/Engine/Plugins/Compositing/Composite/Source/Composite/Public/Passes/CompositePassAlphaTransform.h
@@ -43,15 +43,15 @@ public:
 
 public:
 	/** Alpha premultiplication operation. */
-	UPROPERTY(BlueprintReadWrite, EditAnywhere, Category = "")
+	UPROPERTY(BlueprintReadWrite, EditAnywhere, Interp, Category = "Composite")
 	ECompositeAlphaPremultiplication PremultOp;
 
 	/** Invert the source alpha, for instance if it comes from a scene capture render. */
-	UPROPERTY(BlueprintReadWrite, EditAnywhere, Category = "", AdvancedDisplay)
+	UPROPERTY(BlueprintReadWrite, EditAnywhere, Interp, Category = "Composite", AdvancedDisplay)
 	bool bInvertSourceAlpha;
 
 	/** Output alpha is forced to 1.0, primarily used for correct blending with composite mesh holdouts. */
-	UPROPERTY(BlueprintReadWrite, EditAnywhere, Category = "", AdvancedDisplay)
+	UPROPERTY(BlueprintReadWrite, EditAnywhere, Interp, Category = "Composite", AdvancedDisplay)
 	ECompositeAlphaOverride OverrideOutputAlpha;
 };
 
diff --git a/Engine/Plugins/Compositing/Composite/Source/Composite/Public/Passes/CompositePassBase.h b/Engine/Plugins/Compositing/Composite/Source/Composite/Public/Passes/CompositePassBase.h
--- a/Engine/Plugins/Compositing/Composite/Source/Composite/Public/Passes/CompositePassBase.h
+++ b/Engine/Plugins/Compositing/Composite/Source/Composite/Public/Passes/CompositePassBase.h
@@ -86,12 +86,12 @@ public:
 	
 protected:
 	/** Whether or not the pass is active. */
-	UPROPERTY(EditAnywhere, BlueprintGetter = IsEnabled, BlueprintSetter = SetEnabled, Category = "", meta = (DisplayPriority = "1"))
+	UPROPERTY(EditAnywhere, BlueprintGetter = IsEnabled, BlueprintSetter = SetEnabled, Interp, Category = "Composite", meta = (DisplayPriority = "1"))
 	bool bIsEnabled = true;
 
 #if WITH_EDITORONLY_DATA
 	/** The display name of the pass */
-	UPROPERTY(EditAnywhere, Category = "", meta = (DisplayName="Name", DisplayPriority = "0"))
+	UPROPERTY(EditAnywhere, Category = "Composite", meta = (DisplayName="Name", DisplayPriority = "0"))
 	FString DisplayName;
 #endif
 
diff --git a/Engine/Plugins/Compositing/Composite/Source/Composite/Public/Passes/CompositePassBlur.h b/Engine/Plugins/Compositing/Composite/Source/Composite/Public/Passes/CompositePassBlur.h
--- a/Engine/Plugins/Compositing/Composite/Source/Composite/Public/Passes/CompositePassBlur.h
+++ b/Engine/Plugins/Compositing/Composite/Source/Composite/Public/Passes/CompositePassBlur.h
@@ -34,12 +34,12 @@ public:
 
 public:
     /** Blur pass method. */
-    UPROPERTY(EditAnywhere, BlueprintReadWrite, Category = "")
+    UPROPERTY(EditAnywhere, BlueprintReadWrite, Interp, Category = "Composite")
     ECompositePassBlurMethod Method;
 
 	/** Size of the blur in pixels. */
-	UPROPERTY(EditAnywhere, BlueprintReadWrite, Category = "", meta = (ClampMin = "1", UIMin = "1", UIMax = "64"))
-	int32 Radius;
+	UPROPERTY(EditAnywhere, BlueprintReadWrite, Interp, Category = "Composite", meta = (ClampMin = "1", UIMin = "1", UIMax = "64"))
+	int32 Radius = 1;
 };
 
 #undef UE_API
diff --git a/Engine/Plugins/Compositing/Composite/Source/Composite/Public/Passes/CompositePassCenteredScale.h b/Engine/Plugins/Compositing/Composite/Source/Composite/Public/Passes/CompositePassCenteredScale.h
--- a/Engine/Plugins/Compositing/Composite/Source/Composite/Public/Passes/CompositePassCenteredScale.h
+++ b/Engine/Plugins/Compositing/Composite/Source/Composite/Public/Passes/CompositePassCenteredScale.h
@@ -62,27 +62,27 @@ public:
 
 public:
 	/** Centered scale calculation mode used to rescale a media texture with black bars into an already constrained aspect ratio viewport. */
-	UPROPERTY(EditAnywhere, BlueprintReadWrite, Category = "")
+	UPROPERTY(EditAnywhere, BlueprintReadWrite, Interp, Category = "Composite")
 	ECompositePassScaleMode ScaleMode;
 
 	/** Source container aspect ratio (or resolution). */
-	UPROPERTY(EditAnywhere, BlueprintReadWrite, Category = "", meta = (EditCondition = "ScaleMode == ECompositePassScaleMode::AspectRatio"))
+	UPROPERTY(EditAnywhere, BlueprintReadWrite, Interp, Category = "Composite", meta = (EditCondition = "ScaleMode == ECompositePassScaleMode::AspectRatio"))
 	FVector2f SourceAspectRatio;
 
 	/** Embedded target aspect ratio (or resolution), without black bars. */
-	UPROPERTY(EditAnywhere, BlueprintReadWrite, Category = "", meta = (EditCondition = "ScaleMode == ECompositePassScaleMode::AspectRatio"))
+	UPROPERTY(EditAnywhere, BlueprintReadWrite, Interp, Category = "Composite", meta = (EditCondition = "ScaleMode == ECompositePassScaleMode::AspectRatio"))
 	FVector2f TargetAspectRatio;
 
 	/** Manual scale factor. */
-	UPROPERTY(EditAnywhere, BlueprintReadWrite, Category = "", meta = (EditCondition = "ScaleMode == ECompositePassScaleMode::Manual"))
+	UPROPERTY(EditAnywhere, BlueprintReadWrite, Interp, Category = "Composite", meta = (EditCondition = "ScaleMode == ECompositePassScaleMode::Manual"))
 	FVector2f ScaleFactor;
 
 	/** Uncrop calculation mode used to uncrop a viewport cropped in by overscan. */
-	UPROPERTY(EditAnywhere, BlueprintReadWrite, Category = "")
+	UPROPERTY(EditAnywhere, BlueprintReadWrite, Category = "Composite")
 	ECompositePassOverscanUncropMode OverscanUncropMode;
 
 	/** Manual overscan to uncrop, with values from 0.0 to 1.0 matching the source camera overscan. */
-	UPROPERTY(EditAnywhere, BlueprintReadWrite, Category = "", meta = (EditCondition = "OverscanUncropMode == ECompositePassOverscanUncropMode::Manual"))
+	UPROPERTY(EditAnywhere, BlueprintReadWrite, Interp, Category = "Composite", meta = (EditCondition = "OverscanUncropMode == ECompositePassOverscanUncropMode::Manual", ClampMin = "0.0", ClampMax = "1.0", UIMin = "0.0", UIMax = "1.0"))
 	float Overscan;
 };
 
diff --git a/Engine/Plugins/Compositing/Composite/Source/Composite/Public/Passes/CompositePassColorGrading.h b/Engine/Plugins/Compositing/Composite/Source/Composite/Public/Passes/CompositePassColorGrading.h
--- a/Engine/Plugins/Compositing/Composite/Source/Composite/Public/Passes/CompositePassColorGrading.h
+++ b/Engine/Plugins/Compositing/Composite/Source/Composite/Public/Passes/CompositePassColorGrading.h
@@ -53,15 +53,15 @@ public:
 
 public:
 	/** Color is divided by its alpha channel before processing, and multiplied again after. */
-	UPROPERTY(BlueprintReadWrite, EditAnywhere, Category = "")
+	UPROPERTY(BlueprintReadWrite, EditAnywhere, Category = "Composite")
 	bool bUnpremult;
 
 	/** Color temperature settings. */
-	UPROPERTY(BlueprintReadWrite, EditAnywhere, Category = "", meta = (DisplayName = "Temperature"))
+	UPROPERTY(BlueprintReadWrite, EditAnywhere, Interp, Category = "Composite", meta = (DisplayName = "Temperature"))
 	FCompositeTemperatureSettings TemperatureSettings;
 
 	/** Color grading settings. */
-	UPROPERTY(BlueprintReadWrite, EditAnywhere, Category = "", meta = (ShowOnlyInnerProperties))
+	UPROPERTY(BlueprintReadWrite, EditAnywhere, Interp, Category = "Composite", meta = (ShowOnlyInnerProperties))
 	FColorGradingSettings ColorGradingSettings;
 };
 
@@ -71,19 +71,19 @@ struct FCompositeColorGradingPassSettings
 	GENERATED_USTRUCT_BODY()
 
 	/** Hidden in scene capture. */
-	UPROPERTY(BlueprintReadWrite, EditAnywhere, Category = "")
+	UPROPERTY(BlueprintReadWrite, EditAnywhere, Category = "Composite")
 	bool bEnableColorGrading = false;
 
 	/** Color is divided by its alpha channel before processing, and multiplied again after. */
-	UPROPERTY(BlueprintReadWrite, EditAnywhere, Category = "")
+	UPROPERTY(BlueprintReadWrite, EditAnywhere, Category = "Composite")
 	bool bUnpremult = true;
 
 	/** Color temperature settings. */
-	UPROPERTY(BlueprintReadWrite, EditAnywhere, Category = "", meta = (DisplayName = "Temperature"))
+	UPROPERTY(BlueprintReadWrite, EditAnywhere, Category = "Composite", meta = (DisplayName = "Temperature"))
 	FCompositeTemperatureSettings TemperatureSettings;
 
 	/** Color grading settings. */
-	UPROPERTY(BlueprintReadWrite, EditAnywhere, Category = "", meta = (ShowOnlyInnerProperties))
+	UPROPERTY(BlueprintReadWrite, EditAnywhere, Category = "Composite", meta = (ShowOnlyInnerProperties))
 	FColorGradingSettings ColorGradingSettings;
 };
 
diff --git a/Engine/Plugins/Compositing/Composite/Source/Composite/Public/Passes/CompositePassColorKeyer.h b/Engine/Plugins/Compositing/Composite/Source/Composite/Public/Passes/CompositePassColorKeyer.h
--- a/Engine/Plugins/Compositing/Composite/Source/Composite/Public/Passes/CompositePassColorKeyer.h
+++ b/Engine/Plugins/Compositing/Composite/Source/Composite/Public/Passes/CompositePassColorKeyer.h
@@ -67,75 +67,75 @@ public:
 
 public:
 	/** Type of screen color (required). The keyer works best against red, green or blue backgrounds. */
-	UPROPERTY(EditAnywhere, BlueprintReadWrite, Category = "")
+	UPROPERTY(EditAnywhere, BlueprintReadWrite, Category = "Composite")
 	ECompositeColorKeyerScreenType ScreenType;
 
-	UPROPERTY(EditAnywhere, BlueprintReadWrite, Category = "", meta = (PropertyGroup = "Keyer Settings"))
+	UPROPERTY(EditAnywhere, BlueprintReadWrite, Category = "Composite", meta = (PropertyGroup = "Keyer Settings"))
 	ECompositeKeyerSource KeyerSource;
 	
 	/** Static background key color.*/
-	UPROPERTY(EditAnywhere, BlueprintReadWrite, Category = "", meta = (PropertyGroup = "Keyer Settings", OnlyUpdateOnInteractionEnd, EditCondition = "KeyerSource == ECompositeKeyerSource::Color", EditConditionHides))
+	UPROPERTY(EditAnywhere, BlueprintReadWrite, Interp, Category = "Composite", meta = (PropertyGroup = "Keyer Settings", OnlyUpdateOnInteractionEnd, EditCondition = "KeyerSource == ECompositeKeyerSource::Color", EditConditionHides))
 	FLinearColor KeyColor;
 
 	/**
 	 * Clean plate background for calculating color differences per pixel, instead of the static key color.
 	 * Resolution must match the composite actor render resolution.
 	 */
-	UPROPERTY(EditAnywhere, BlueprintReadWrite, Category = "", meta = (PropertyGroup = "Keyer Settings", OnlyUpdateOnInteractionEnd, EditCondition = "KeyerSource == ECompositeKeyerSource::CleanPlate", EditConditionHides))
+	UPROPERTY(EditAnywhere, BlueprintReadWrite, Category = "Composite", meta = (PropertyGroup = "Keyer Settings", OnlyUpdateOnInteractionEnd, EditCondition = "KeyerSource == ECompositeKeyerSource::CleanPlate", EditConditionHides))
 	TObjectPtr<UTexture> CleanPlate;
 
 #if WITH_EDITORONLY_DATA
 	/** Setting that controls how long to delay before capturing a clean plate */
-	UPROPERTY(EditAnywhere, Category = "", meta = (PropertyGroup = "Keyer Settings", EditCondition = "KeyerSource == ECompositeKeyerSource::CleanPlate", EditConditionHides, ClampMin=0.0, UIMin=0.0, ClampMax=60.0, UIMax=60.0))
+	UPROPERTY(EditAnywhere, Category = "Composite", meta = (PropertyGroup = "Keyer Settings", EditCondition = "KeyerSource == ECompositeKeyerSource::CleanPlate", EditConditionHides, ClampMin=0.0, UIMin=0.0, ClampMax=60.0, UIMax=60.0))
 	float CleanPlateCountdown = 0.0f;
 
 	/** Setting that determines how large to display the plate texture 'picture-in-picture' preview in the level editor when using the eye dropper color picker */
-	UPROPERTY(EditAnywhere, Category = "", meta = (PropertyGroup = "Keyer Settings", EditCondition = "KeyerSource == ECompositeKeyerSource::Color", EditConditionHides, ClampMin=1.0, UIMin=1.0, ClampMax=10.0, UIMax=10.0))
+	UPROPERTY(EditAnywhere, Category = "Composite", meta = (PropertyGroup = "Keyer Settings", EditCondition = "KeyerSource == ECompositeKeyerSource::Color", EditConditionHides, ClampMin=1.0, UIMin=1.0, ClampMax=10.0, UIMax=10.0))
 	float PlateTexturePreviewSize = 5.0f;
 #endif
 	
 	/** Weight of the foreground red channel contributing to the key matte hardness. */
-	UPROPERTY(EditAnywhere, BlueprintReadWrite, Category = "", meta = (UIMin = 0.0f, UIMax = 1.0f, EditCondition = "ScreenType != ECompositeColorKeyerScreenType::Red"))
+	UPROPERTY(EditAnywhere, BlueprintReadWrite, Interp, Category = "Composite", meta = (UIMin = 0.0f, UIMax = 1.0f, EditCondition = "ScreenType != ECompositeColorKeyerScreenType::Red"))
 	float RedWeight;
 
 	/** Weight of the foreground green channel contributing to the key matte hardness. */
-	UPROPERTY(EditAnywhere, BlueprintReadWrite, Category = "", meta = (UIMin = 0.0f, UIMax = 1.0f, EditCondition = "ScreenType != ECompositeColorKeyerScreenType::Green"))
+	UPROPERTY(EditAnywhere, BlueprintReadWrite, Interp, Category = "Composite", meta = (UIMin = 0.0f, UIMax = 1.0f, EditCondition = "ScreenType != ECompositeColorKeyerScreenType::Green"))
 	float GreenWeight;
 
 	/** Weight of the foreground blue channel contributing to the key matte hardness. */
-	UPROPERTY(EditAnywhere, BlueprintReadWrite, Category = "", meta = (UIMin = 0.0f, UIMax = 1.0f, EditCondition = "ScreenType != ECompositeColorKeyerScreenType::Blue"))
+	UPROPERTY(EditAnywhere, BlueprintReadWrite, Interp, Category = "Composite", meta = (UIMin = 0.0f, UIMax = 1.0f, EditCondition = "ScreenType != ECompositeColorKeyerScreenType::Blue"))
 	float BlueWeight;
 
 	/** Thresholds any alpha value outside the specified range to zero or one respectively, with linear interpolation in-between. */
-	UPROPERTY(EditAnywhere, BlueprintReadWrite, Category = "", meta = (UIMin = 0.0f, UIMax = 1.0f))
+	UPROPERTY(EditAnywhere, BlueprintReadWrite, Interp, Category = "Composite", meta = (UIMin = 0.0f, UIMax = 1.0f))
 	FVector2f AlphaThreshold;
 
 	/** Strength of the spill reduction, 0.0: none, 1.0: full. */
-	UPROPERTY(EditAnywhere, BlueprintReadWrite, Category = "", meta = (UIMin = 0.0f, UIMax = 1.0f))
+	UPROPERTY(EditAnywhere, BlueprintReadWrite, Interp, Category = "Composite", meta = (UIMin = 0.0f, UIMax = 1.0f))
 	float DespillStrength;
 
 	/** Strength of the vignette removal. Used to improve plate uniformity & remove darker corners. */
-	UPROPERTY(EditAnywhere, BlueprintReadWrite, Category = "", meta = (UIMin = 0.0f, UIMax = 1.0f))
+	UPROPERTY(EditAnywhere, BlueprintReadWrite, Interp, Category = "Composite", meta = (UIMin = 0.0f, UIMax = 1.0f))
 	float DevignetteStrength;
 
 	/** When enabled, we undo devignetting before outputting the keyed plate. */
-	UPROPERTY(EditAnywhere, BlueprintReadWrite, Category = "", AdvancedDisplay)
+	UPROPERTY(EditAnywhere, BlueprintReadWrite, Category = "Composite", AdvancedDisplay)
 	bool bPreserveVignetteAfterKey;
 
 	/** Denoising method applied before the keyer. */
-	UPROPERTY(EditAnywhere, BlueprintReadWrite, Category = "", AdvancedDisplay)
+	UPROPERTY(EditAnywhere, BlueprintReadWrite, Category = "Composite", AdvancedDisplay)
 	ECompositeDenoiseMethod DenoiseMethod;
 
 	/** Vizualize the alpha key or fill. */
-	UPROPERTY(EditAnywhere, BlueprintReadWrite, Category = "")
+	UPROPERTY(EditAnywhere, BlueprintReadWrite, Category = "Composite")
 	ECompositeColorKeyerVisualization Visualization;
 
 	/** Only apply the despill algorithm controlled by DespillStrength, alpha is unmodified. */
-	UPROPERTY(EditAnywhere, BlueprintReadWrite, Category = "", AdvancedDisplay)
+	UPROPERTY(EditAnywhere, BlueprintReadWrite, Category = "Composite", AdvancedDisplay)
 	bool bDespillOnly;
 
 	/** Invert the alpha key. */
-	UPROPERTY(EditAnywhere, BlueprintReadWrite, Category = "", AdvancedDisplay)
+	UPROPERTY(EditAnywhere, BlueprintReadWrite, Category = "Composite", AdvancedDisplay)
 	bool bInvertAlpha;
 
 #if WITH_EDITOR
diff --git a/Engine/Plugins/Compositing/Composite/Source/Composite/Public/Passes/CompositePassDistortion.h b/Engine/Plugins/Compositing/Composite/Source/Composite/Public/Passes/CompositePassDistortion.h
--- a/Engine/Plugins/Compositing/Composite/Source/Composite/Public/Passes/CompositePassDistortion.h
+++ b/Engine/Plugins/Compositing/Composite/Source/Composite/Public/Passes/CompositePassDistortion.h
@@ -41,7 +41,7 @@ public:
 
 public:
 	/** Distortion to apply on input. */
-	UPROPERTY(EditAnywhere, BlueprintReadWrite, Category = "")
+	UPROPERTY(EditAnywhere, BlueprintReadWrite, Category = "Composite")
 	ECompositeDistortion Distortion;
 
 private:
diff --git a/Engine/Plugins/Compositing/Composite/Source/Composite/Public/Passes/CompositePassFXAA.h b/Engine/Plugins/Compositing/Composite/Source/Composite/Public/Passes/CompositePassFXAA.h
--- a/Engine/Plugins/Compositing/Composite/Source/Composite/Public/Passes/CompositePassFXAA.h
+++ b/Engine/Plugins/Compositing/Composite/Source/Composite/Public/Passes/CompositePassFXAA.h
@@ -26,7 +26,7 @@ public:
 public:
 
 	/** FXAA quality control, matching the "r.fxaa.quality" console variable values. */
-	UPROPERTY(EditAnywhere, BlueprintReadWrite, Category = "", meta=(ClampMin="0", ClampMax="5", UIMin = "0", UIMax = "5"))
+	UPROPERTY(EditAnywhere, BlueprintReadWrite, Interp, Category = "Composite", meta=(ClampMin="0", ClampMax="5", UIMin = "0", UIMax = "5"))
 	int32 Quality;
 };
 
diff --git a/Engine/Plugins/Compositing/Composite/Source/Composite/Public/Passes/CompositePassMasking.h b/Engine/Plugins/Compositing/Composite/Source/Composite/Public/Passes/CompositePassMasking.h
--- a/Engine/Plugins/Compositing/Composite/Source/Composite/Public/Passes/CompositePassMasking.h
+++ b/Engine/Plugins/Compositing/Composite/Source/Composite/Public/Passes/CompositePassMasking.h
@@ -37,19 +37,19 @@ public:
 
 public:
 	/** Masking texture. */
-	UPROPERTY(EditAnywhere, BlueprintReadWrite, Category = "", meta = (AllowedClasses = "/Script/Engine.Texture2D, /Script/MediaAssets.MediaTexture"))
+	UPROPERTY(EditAnywhere, BlueprintReadWrite, Category = "Composite", meta = (AllowedClasses = "/Script/Engine.Texture2D, /Script/MediaAssets.MediaTexture"))
 	TObjectPtr<UTexture> MaskTexture;
 
 	/** Color channel used as the mask source. */
-	UPROPERTY(BlueprintReadWrite, EditAnywhere, Category = "")
+	UPROPERTY(BlueprintReadWrite, EditAnywhere, Interp, Category = "Composite")
 	ECompositeColorChannel MaskSource;
 
 	/** Color is divided by its alpha channel before masking, and multiplied again after. */
-	UPROPERTY(BlueprintReadWrite, EditAnywhere, Category = "")
+	UPROPERTY(BlueprintReadWrite, EditAnywhere, Interp, Category = "Composite")
 	bool bUnpremult;
 
 	/** Invert mask alpha. */
-	UPROPERTY(BlueprintReadWrite, EditAnywhere, Category = "")
+	UPROPERTY(BlueprintReadWrite, EditAnywhere, Interp, Category = "Composite")
 	bool bInvert;
 };
 
diff --git a/Engine/Plugins/Compositing/Composite/Source/Composite/Public/Passes/CompositePassMaterial.h b/Engine/Plugins/Compositing/Composite/Source/Composite/Public/Passes/CompositePassMaterial.h
--- a/Engine/Plugins/Compositing/Composite/Source/Composite/Public/Passes/CompositePassMaterial.h
+++ b/Engine/Plugins/Compositing/Composite/Source/Composite/Public/Passes/CompositePassMaterial.h
@@ -29,7 +29,7 @@ public:
 
 public:
 	/** Post-process material to execute. Input0 is connected to SceneTexture's PostprocessInput0. */
-	UPROPERTY(EditAnywhere, BlueprintReadWrite, Category = "")
+	UPROPERTY(EditAnywhere, BlueprintReadWrite, Category = "Composite")
 	TObjectPtr<UMaterialInterface> PostProcessMaterial;
 };
 
diff --git a/Engine/Plugins/Compositing/Composite/Source/Composite/Public/Passes/CompositePassOpenColorIO.h b/Engine/Plugins/Compositing/Composite/Source/Composite/Public/Passes/CompositePassOpenColorIO.h
--- a/Engine/Plugins/Compositing/Composite/Source/Composite/Public/Passes/CompositePassOpenColorIO.h
+++ b/Engine/Plugins/Compositing/Composite/Source/Composite/Public/Passes/CompositePassOpenColorIO.h
@@ -28,7 +28,7 @@ public:
 
 public:
 	/** Color conversion settings. */
-	UPROPERTY(EditAnywhere, BlueprintReadWrite, Category = "")
+	UPROPERTY(EditAnywhere, BlueprintReadWrite, Category = "Composite")
 	FOpenColorIOColorConversionSettings ColorConversionSettings;
 };
 
diff --git a/Engine/Plugins/Compositing/Composite/Source/Composite/Public/Passes/CompositePassSMAA.h b/Engine/Plugins/Compositing/Composite/Source/Composite/Public/Passes/CompositePassSMAA.h
--- a/Engine/Plugins/Compositing/Composite/Source/Composite/Public/Passes/CompositePassSMAA.h
+++ b/Engine/Plugins/Compositing/Composite/Source/Composite/Public/Passes/CompositePassSMAA.h
@@ -26,7 +26,7 @@ public:
 public:
 
 	/** SMAA quality control, matching the "r.SMAA.quality" console variable values. */
-	UPROPERTY(EditAnywhere, BlueprintReadWrite, Category = "", meta=(ClampMin="0", ClampMax="5", UIMin = "0", UIMax = "5"))
+	UPROPERTY(EditAnywhere, BlueprintReadWrite, Interp, Category = "Composite", meta=(ClampMin="0", ClampMax="5", UIMin = "0", UIMax = "5"))
 	int32 Quality;
 };
 
diff --git a/Engine/Plugins/Compositing/Composite/Source/CompositeEditor/CompositeEditor.Build.cs b/Engine/Plugins/Compositing/Composite/Source/CompositeEditor/CompositeEditor.Build.cs
--- a/Engine/Plugins/Compositing/Composite/Source/CompositeEditor/CompositeEditor.Build.cs
+++ b/Engine/Plugins/Compositing/Composite/Source/CompositeEditor/CompositeEditor.Build.cs
@@ -33,11 +33,15 @@ public class CompositeEditor : ModuleRules
 				"LevelEditor",
 				"MediaAssets",
 				"MediaFrameworkUtilities",
+				"MovieScene",
+				"MovieSceneTracks",
+				"MovieSceneTools",
 				"ObjectMixerEditor",
 				"Projects",
 				"PropertyEditor",
 				"UnrealEd",
 				"SceneOutliner",
+				"Sequencer",
                 "Slate",
                 "SlateCore",
                 "ToolMenus",
diff --git a/Engine/Plugins/Compositing/Composite/Source/CompositeEditor/Private/CompositeEditorModule.cpp b/Engine/Plugins/Compositing/Composite/Source/CompositeEditor/Private/CompositeEditorModule.cpp
--- a/Engine/Plugins/Compositing/Composite/Source/CompositeEditor/Private/CompositeEditorModule.cpp
+++ b/Engine/Plugins/Compositing/Composite/Source/CompositeEditor/Private/CompositeEditorModule.cpp
@@ -2,6 +2,9 @@
 
 #include "CompositeEditorModule.h"
 
+#include "ISequencerModule.h"
+#include "Sequencer/CompositeActorObjectSchema.h"
+
 #include "ColorGradingEditorDataModel.h"
 #include "ColorGradingMixerObjectFilterRegistry.h"
 #include "Editor.h"
@@ -20,12 +23,14 @@
 #include "ColorGradingDrawer/ColorGradingHierarchyConfig_Composite.h"
 #include "ColorGradingDrawer/ColorGradingSelectionHandler_Composite.h"
 #include "Customizations/CompositeActorCustomization.h"
+#include "Customizations/CompositeLayerBaseCustomization.h"
 #include "Customizations/CompositeLayerPlateCustomization.h"
 #include "Customizations/CompositeLayerSceneCaptureCustomization.h"
 #include "Customizations/CompositeLayerShadowReflectionCustomization.h"
 #include "Customizations/CompositeLayerSingleLightShadowCustomization.h"
 #include "Customizations/CompositeDepthMeshComponentCustomization.h"
 #include "Customizations/CompositePassColorKeyerCustomization.h"
+#include "Layers/CompositeLayerBase.h"
 #include "Layers/CompositeLayerPlate.h"
 #include "Layers/CompositeLayerSceneCapture.h"
 #include "Layers/CompositeLayerShadowReflection.h"
@@ -63,10 +68,24 @@ void FCompositeEditorModule::StartupModule()
 	);
 
 	FColorGradingMixerObjectFilterRegistry::RegisterSelectionHandler(FGetSelectionHandler::CreateStatic(&FColorGradingSelectionHandler_Composite::MakeInstance));
+
+	// Register Sequencer object schema for layer/pass sub-object discovery.
+	ISequencerModule& SequencerModule = FModuleManager::LoadModuleChecked<ISequencerModule>("Sequencer");
+	CompositeActorObjectSchema = MakeShared<UE::Sequencer::FCompositeActorObjectSchema>();
+	SequencerModule.RegisterObjectSchema(CompositeActorObjectSchema);
 }
 
 void FCompositeEditorModule::ShutdownModule()
 {
+	if (CompositeActorObjectSchema)
+	{
+		if (ISequencerModule* SequencerModule = FModuleManager::GetModulePtr<ISequencerModule>("Sequencer"))
+		{
+			SequencerModule->UnregisterObjectSchema(CompositeActorObjectSchema);
+		}
+		CompositeActorObjectSchema.Reset();
+	}
+
 	UnregisterCustomizations();
 	UnregisterTabSpawners();
 
@@ -183,6 +202,11 @@ void FCompositeEditorModule::UnregisterTabSpawners()
 void FCompositeEditorModule::RegisterCustomizations()
 {
 	FPropertyEditorModule& PropertyModule = FModuleManager::LoadModuleChecked<FPropertyEditorModule>("PropertyEditor");
+	PropertyModule.RegisterCustomClassLayout(
+		UCompositeLayerBase::StaticClass()->GetFName(),
+		FOnGetDetailCustomizationInstance::CreateStatic(&FCompositeLayerBaseCustomization::MakeInstance)
+	);
+
 	PropertyModule.RegisterCustomClassLayout(
 		UCompositeLayerPlate::StaticClass()->GetFName(),
 		FOnGetDetailCustomizationInstance::CreateStatic(&FCompositeLayerPlateCustomization::MakeInstance)
@@ -226,6 +250,7 @@ void FCompositeEditorModule::UnregisterCustomizations()
 		if (FPropertyEditorModule* PropertyModule = FModuleManager::Get().GetModulePtr<FPropertyEditorModule>("PropertyEditor"))
 		{
 			PropertyModule->UnregisterCustomClassLayout(ACompositeActor::StaticClass()->GetFName());
+			PropertyModule->UnregisterCustomClassLayout(UCompositeLayerBase::StaticClass()->GetFName());
 			PropertyModule->UnregisterCustomClassLayout(UCompositeLayerPlate::StaticClass()->GetFName());
 			PropertyModule->UnregisterCustomClassLayout(UCompositeLayerShadowReflection::StaticClass()->GetFName());
 			PropertyModule->UnregisterCustomClassLayout(UCompositeLayerSceneCapture::StaticClass()->GetFName());
diff --git a/Engine/Plugins/Compositing/Composite/Source/CompositeEditor/Private/Customizations/CompositeLayerBaseCustomization.cpp b/Engine/Plugins/Compositing/Composite/Source/CompositeEditor/Private/Customizations/CompositeLayerBaseCustomization.cpp
new file mode 100644
--- /dev/null
+++ b/Engine/Plugins/Compositing/Composite/Source/CompositeEditor/Private/Customizations/CompositeLayerBaseCustomization.cpp
@@ -0,0 +1,25 @@
+// Copyright Epic Games, Inc. All Rights Reserved.
+
+#include "CompositeLayerBaseCustomization.h"
+
+#include "DetailCategoryBuilder.h"
+#include "DetailLayoutBuilder.h"
+
+TSharedRef<IDetailCustomization> FCompositeLayerBaseCustomization::MakeInstance()
+{
+	return MakeShared<FCompositeLayerBaseCustomization>();
+}
+
+void FCompositeLayerBaseCustomization::CustomizeDetails(IDetailLayoutBuilder& DetailLayout)
+{
+	TArray<TWeakObjectPtr<UObject>> Objects;
+	DetailLayout.GetObjectsBeingCustomized(Objects);
+
+	if (Objects.Num() > 0)
+	{
+		if (TStrongObjectPtr<UObject> PinnedObject = Objects[0].Pin())
+		{
+			DetailLayout.EditCategory("Composite", PinnedObject->GetClass()->GetDisplayNameText());
+		}
+	}
+}
diff --git a/Engine/Plugins/Compositing/Composite/Source/CompositeEditor/Private/Customizations/CompositeLayerBaseCustomization.h b/Engine/Plugins/Compositing/Composite/Source/CompositeEditor/Private/Customizations/CompositeLayerBaseCustomization.h
new file mode 100644
--- /dev/null
+++ b/Engine/Plugins/Compositing/Composite/Source/CompositeEditor/Private/Customizations/CompositeLayerBaseCustomization.h
@@ -0,0 +1,17 @@
+// Copyright Epic Games, Inc. All Rights Reserved.
+
+#pragma once
+
+#include "IDetailCustomization.h"
+
+/**
+ * Base customization for all derivations of UCompositeLayerBase.
+ * Renames the "Composite" category to the layer class display name (e.g. "Main Render", "Processing").
+ */
+class FCompositeLayerBaseCustomization : public IDetailCustomization
+{
+public:
+	static TSharedRef<IDetailCustomization> MakeInstance();
+
+	virtual void CustomizeDetails(IDetailLayoutBuilder& DetailLayout) override;
+};
diff --git a/Engine/Plugins/Compositing/Composite/Source/CompositeEditor/Private/Customizations/CompositeLayerPlateCustomization.cpp b/Engine/Plugins/Compositing/Composite/Source/CompositeEditor/Private/Customizations/CompositeLayerPlateCustomization.cpp
--- a/Engine/Plugins/Compositing/Composite/Source/CompositeEditor/Private/Customizations/CompositeLayerPlateCustomization.cpp
+++ b/Engine/Plugins/Compositing/Composite/Source/CompositeEditor/Private/Customizations/CompositeLayerPlateCustomization.cpp
@@ -68,7 +68,7 @@ void FCompositeLayerPlateCustomization::CustomizeDetails(IDetailLayoutBuilder& D
 		DetailLayout.HideProperty(PropertyName);
 	}
 	
-	IDetailCategoryBuilder& PlateCategory = DetailLayout.EditCategoryAllowNone(NAME_None);
+	IDetailCategoryBuilder& PlateCategory = DetailLayout.EditCategory("Composite", UCompositeLayerPlate::StaticClass()->GetDisplayNameText());
 
 	TMap<FName, CompositeLayerPlateCustomization::FOnCustomizePropertyRow> PropertyCustomizations =
 	{
diff --git a/Engine/Plugins/Compositing/Composite/Source/CompositeEditor/Private/Customizations/CompositeLayerSceneCaptureCustomization.cpp b/Engine/Plugins/Compositing/Composite/Source/CompositeEditor/Private/Customizations/CompositeLayerSceneCaptureCustomization.cpp
--- a/Engine/Plugins/Compositing/Composite/Source/CompositeEditor/Private/Customizations/CompositeLayerSceneCaptureCustomization.cpp
+++ b/Engine/Plugins/Compositing/Composite/Source/CompositeEditor/Private/Customizations/CompositeLayerSceneCaptureCustomization.cpp
@@ -48,7 +48,7 @@ void FCompositeLayerSceneCaptureCustomization::CustomizeDetails(IDetailLayoutBui
 	
 	TArray<TWeakObjectPtr<UCompositeLayerSceneCapture>> Objects = DetailLayout.GetObjectsOfTypeBeingCustomized<UCompositeLayerSceneCapture>();
 
-	IDetailCategoryBuilder& LayerCategory = DetailLayout.EditCategoryAllowNone(NAME_None);
+	IDetailCategoryBuilder& LayerCategory = DetailLayout.EditCategory("Composite", UCompositeLayerSceneCapture::StaticClass()->GetDisplayNameText());
 	
 	// Add all default simple properties ahead of the custom group
 	TArray<TSharedRef<IPropertyHandle>> LayerProperties;
diff --git a/Engine/Plugins/Compositing/Composite/Source/CompositeEditor/Private/Customizations/CompositeLayerShadowReflectionCustomization.cpp b/Engine/Plugins/Compositing/Composite/Source/CompositeEditor/Private/Customizations/CompositeLayerShadowReflectionCustomization.cpp
--- a/Engine/Plugins/Compositing/Composite/Source/CompositeEditor/Private/Customizations/CompositeLayerShadowReflectionCustomization.cpp
+++ b/Engine/Plugins/Compositing/Composite/Source/CompositeEditor/Private/Customizations/CompositeLayerShadowReflectionCustomization.cpp
@@ -40,7 +40,7 @@ void FCompositeLayerShadowReflectionCustomization::CustomizeDetails(IDetailLayou
 	
 	TArray<TWeakObjectPtr<UCompositeLayerShadowReflection>> Objects = DetailLayout.GetObjectsOfTypeBeingCustomized<UCompositeLayerShadowReflection>();
 
-	IDetailCategoryBuilder& LayerCategory = DetailLayout.EditCategoryAllowNone(NAME_None);
+	IDetailCategoryBuilder& LayerCategory = DetailLayout.EditCategory("Composite", UCompositeLayerShadowReflection::StaticClass()->GetDisplayNameText());
 	
 	// Add all default simple properties ahead of the custom group
 	TArray<TSharedRef<IPropertyHandle>> PlateProperties;
diff --git a/Engine/Plugins/Compositing/Composite/Source/CompositeEditor/Private/Customizations/CompositeLayerSingleLightShadowCustomization.cpp b/Engine/Plugins/Compositing/Composite/Source/CompositeEditor/Private/Customizations/CompositeLayerSingleLightShadowCustomization.cpp
--- a/Engine/Plugins/Compositing/Composite/Source/CompositeEditor/Private/Customizations/CompositeLayerSingleLightShadowCustomization.cpp
+++ b/Engine/Plugins/Compositing/Composite/Source/CompositeEditor/Private/Customizations/CompositeLayerSingleLightShadowCustomization.cpp
@@ -37,7 +37,7 @@ void FCompositeLayerSingleLightShadowCustomization::CustomizeDetails(IDetailLayo
 	
 	TArray<TWeakObjectPtr<UCompositeLayerSingleLightShadow>> Objects = DetailLayout.GetObjectsOfTypeBeingCustomized<UCompositeLayerSingleLightShadow>();
 
-	IDetailCategoryBuilder& LayerCategory = DetailLayout.EditCategoryAllowNone(NAME_None);
+	IDetailCategoryBuilder& LayerCategory = DetailLayout.EditCategory("Composite", UCompositeLayerSingleLightShadow::StaticClass()->GetDisplayNameText());
 
 	// Mark ColorGrading as customized so that it can be placed below 'Shadow Casting Content' custom group
 	TSharedPtr<IPropertyHandle> ColorGradingHandle = DetailLayout.GetProperty(GET_MEMBER_NAME_CHECKED(UCompositeLayerSingleLightShadow, ColorGrading));
diff --git a/Engine/Plugins/Compositing/Composite/Source/CompositeEditor/Private/Customizations/CompositePassBaseCustomization.cpp b/Engine/Plugins/Compositing/Composite/Source/CompositeEditor/Private/Customizations/CompositePassBaseCustomization.cpp
--- a/Engine/Plugins/Compositing/Composite/Source/CompositeEditor/Private/Customizations/CompositePassBaseCustomization.cpp
+++ b/Engine/Plugins/Compositing/Composite/Source/CompositeEditor/Private/Customizations/CompositePassBaseCustomization.cpp
@@ -15,6 +15,18 @@ TSharedRef<IDetailCustomization> FCompositePassBaseCustomization::MakeInstance()
 
 void FCompositePassBaseCustomization::CustomizeDetails(IDetailLayoutBuilder& DetailLayout)
 {
+	// Rename "Composite" category to the pass class display name (e.g. "Blur", "Color Grading")
+	TArray<TWeakObjectPtr<UObject>> Objects;
+	DetailLayout.GetObjectsBeingCustomized(Objects);
+
+	if (Objects.Num() > 0)
+	{
+		if (TStrongObjectPtr<UObject> PinnedObject = Objects[0].Pin())
+		{
+			DetailLayout.EditCategory("Composite", PinnedObject->GetClass()->GetDisplayNameText());
+		}
+	}
+
 	CustomizeIsEnabledProperty(DetailLayout);
 }
 
@@ -46,7 +58,7 @@ void FCompositePassBaseCustomization::CustomizeIsEnabledProperty(IDetailLayoutBu
 	}
 	else
 	{
-		IDetailCategoryBuilder& DefaultCategory = DetailLayout.EditCategoryAllowNone(NAME_None);
+		IDetailCategoryBuilder& DefaultCategory = DetailLayout.EditCategory("Composite");
 		PropertyRow = &DefaultCategory.AddProperty(EnabledHandle);
 	}
 
diff --git a/Engine/Plugins/Compositing/Composite/Source/CompositeEditor/Private/Customizations/CompositePassColorKeyerCustomization.cpp b/Engine/Plugins/Compositing/Composite/Source/CompositeEditor/Private/Customizations/CompositePassColorKeyerCustomization.cpp
--- a/Engine/Plugins/Compositing/Composite/Source/CompositeEditor/Private/Customizations/CompositePassColorKeyerCustomization.cpp
+++ b/Engine/Plugins/Compositing/Composite/Source/CompositeEditor/Private/Customizations/CompositePassColorKeyerCustomization.cpp
@@ -424,8 +424,7 @@ void FCompositePassColorKeyerCustomization::CustomizeDetails(IDetailLayoutBuilde
 			return MakeShared<FKeyerColorStructCustomization>();
 		}));
 	
-	// All UCompositePassColorKeyer properties reside in the default category (Category = ""). If this changes, will need to update this category 
-	IDetailCategoryBuilder& DefaultCategory = DetailBuilder.EditCategoryAllowNone(NAME_None);
+	IDetailCategoryBuilder& DefaultCategory = DetailBuilder.EditCategory("Composite", UCompositePassColorKeyer::StaticClass()->GetDisplayNameText());
 
 	// 'Customize' properties that should be above the 'Keyer Settings' custom group
 	DefaultCategory.AddProperty(DetailBuilder.GetProperty(GET_MEMBER_NAME_CHECKED(UCompositePassColorKeyer, DisplayName), UCompositePassBase::StaticClass()));
diff --git a/Engine/Plugins/Compositing/Composite/Source/CompositeEditor/Private/Sequencer/CompositeActorObjectSchema.cpp b/Engine/Plugins/Compositing/Composite/Source/CompositeEditor/Private/Sequencer/CompositeActorObjectSchema.cpp
new file mode 100644
--- /dev/null
+++ b/Engine/Plugins/Compositing/Composite/Source/CompositeEditor/Private/Sequencer/CompositeActorObjectSchema.cpp
@@ -0,0 +1,387 @@
+// Copyright Epic Games, Inc. All Rights Reserved.
+
+#include "CompositeActorObjectSchema.h"
+
+#include "CompositeActor.h"
+#include "Layers/CompositeLayerBase.h"
+#include "Passes/CompositePassBase.h"
+
+#include "GameFramework/Actor.h"
+#include "Components/ActorComponent.h"
+#include "UObject/UnrealType.h"
+#include "ISequencer.h"
+#include "ISequencerModule.h"
+#include "Framework/MultiBox/MultiBoxBuilder.h"
+#include "Framework/MultiBox/MultiBoxExtender.h"
+#include "ScopedTransaction.h"
+
+#define LOCTEXT_NAMESPACE "CompositeActorObjectSchema"
+
+namespace UE::Sequencer
+{
+
+UObject* FCompositeActorObjectSchema::GetParentObject(UObject* Object) const
+{
+	if (!Object)
+	{
+		return nullptr;
+	}
+
+	// Component -> Actor (handle components so they still work under our schema)
+	if (UActorComponent* Component = Cast<UActorComponent>(Object))
+	{
+		return Component->GetOwner();
+	}
+
+	// Layer -> Actor (check layer first since UCompositeLayerBase inherits UCompositePassBase)
+	if (Object->IsA<UCompositeLayerBase>())
+	{
+		return Object->GetTypedOuter<ACompositeActor>();
+	}
+
+	// Pass -> Layer
+	if (Object->IsA<UCompositePassBase>())
+	{
+		return Object->GetTypedOuter<UCompositeLayerBase>();
+	}
+
+	return nullptr;
+}
+
+FObjectSchemaRelevancy FCompositeActorObjectSchema::GetRelevancy(const UObject* InObject) const
+{
+	if (!InObject)
+	{
+		return FObjectSchemaRelevancy();
+	}
+
+	// Claim ACompositeActor with its specific class — beats FActorSchema's AActor::StaticClass()
+	// because ACompositeActor::IsChildOf(AActor) is true (more derived wins).
+	if (InObject->IsA<ACompositeActor>())
+	{
+		return FObjectSchemaRelevancy(ACompositeActor::StaticClass());
+	}
+
+	// Also claim components owned by a CompositeActor so they nest correctly under our actor schema.
+	if (const UActorComponent* Component = Cast<UActorComponent>(InObject))
+	{
+		if (Component->GetOwner() && Component->GetOwner()->IsA<ACompositeActor>())
+		{
+			return FObjectSchemaRelevancy(UActorComponent::StaticClass(), 1);
+		}
+	}
+
+	if (InObject->IsA<UCompositeLayerBase>())
+	{
+		return FObjectSchemaRelevancy(UCompositeLayerBase::StaticClass(), 1);
+	}
+
+	if (InObject->IsA<UCompositePassBase>())
+	{
+		return FObjectSchemaRelevancy(UCompositePassBase::StaticClass(), 1);
+	}
+
+	return FObjectSchemaRelevancy();
+}
+
+FText FCompositeActorObjectSchema::GetPrettyName(const UObject* Object) const
+{
+	if (!Object)
+	{
+		return FText::GetEmpty();
+	}
+
+	if (const AActor* Actor = Cast<AActor>(Object))
+	{
+		return FText::FromString(Actor->GetActorLabel());
+	}
+
+	// Note: UCompositeLayerBase inherits from UCompositePassBase now.
+	if (const UCompositePassBase* Pass = Cast<UCompositePassBase>(Object))
+	{
+		const FString Name = Pass->GetDisplayName();
+		if (!Name.IsEmpty())
+		{
+			return FText::FromString(Name);
+		}
+	}
+
+	return Object->GetClass()->GetDisplayNameText();
+}
+
+TSharedPtr<FExtender> FCompositeActorObjectSchema::ExtendObjectBindingMenu(
+	TSharedRef<FUICommandList> CommandList,
+	TWeakPtr<ISequencer> WeakSequencer,
+	TArrayView<UObject* const> ContextSensitiveObjects) const
+{
+	// Actor context: show layers and components
+	TArray<ACompositeActor*> Actors;
+	for (UObject* Object : ContextSensitiveObjects)
+	{
+		if (ACompositeActor* Actor = Cast<ACompositeActor>(Object))
+		{
+			Actors.Add(Actor);
+		}
+	}
+
+	if (Actors.Num() > 0)
+	{
+		TSharedRef<FExtender> Extender = MakeShared<FExtender>();
+		Extender->AddMenuExtension(
+			SequencerMenuExtensionPoints::AddTrackMenu_PropertiesSection,
+			EExtensionHook::Before,
+			CommandList,
+			FMenuExtensionDelegate::CreateRaw(this, &FCompositeActorObjectSchema::HandleActorTrackMenu, WeakSequencer, Actors));
+		
+		return Extender;
+	}
+
+	// Layer context: show passes
+	TArray<UCompositeLayerBase*> Layers;
+	for (UObject* Object : ContextSensitiveObjects)
+	{
+		if (UCompositeLayerBase* Layer = Cast<UCompositeLayerBase>(Object))
+		{
+			Layers.Add(Layer);
+		}
+	}
+
+	if (Layers.Num() > 0)
+	{
+		TSharedRef<FExtender> Extender = MakeShared<FExtender>();
+		Extender->AddMenuExtension(
+			SequencerMenuExtensionPoints::AddTrackMenu_PropertiesSection,
+			EExtensionHook::Before,
+			CommandList,
+			FMenuExtensionDelegate::CreateRaw(this, &FCompositeActorObjectSchema::HandleLayerTrackMenu, WeakSequencer, Layers));
+		
+		return Extender;
+	}
+
+	return nullptr;
+}
+
+void FCompositeActorObjectSchema::HandleActorTrackMenu(
+	FMenuBuilder& MenuBuilder,
+	TWeakPtr<ISequencer> WeakSequencer,
+	TArray<ACompositeActor*> Actors) const
+{
+	TSharedPtr<ISequencer> Sequencer = WeakSequencer.Pin();
+	if (!Sequencer)
+	{
+		return;
+	}
+
+	// --- Composite Layers section ---
+	MenuBuilder.BeginSection("CompositeLayers", LOCTEXT("CompositeLayersSection", "Composite Layers"));
+	{
+		for (ACompositeActor* Actor : Actors)
+		{
+			TArray<UCompositeLayerBase*> Layers = Actor->GetCompositeLayers();
+			for (UCompositeLayerBase* Layer : Layers)
+			{
+				if (!Layer)
+				{
+					continue;
+				}
+
+				if (Sequencer->GetHandleToObject(Layer, false).IsValid())
+				{
+					continue;
+				}
+
+				FText DisplayName;
+				const FString LayerName = Layer->GetDisplayName();
+				if (!LayerName.IsEmpty())
+				{
+					DisplayName = FText::FromString(LayerName);
+				}
+				else
+				{
+					DisplayName = Layer->GetClass()->GetDisplayNameText();
+				}
+
+				FUIAction Action(FExecuteAction::CreateRaw(
+					this, &FCompositeActorObjectSchema::HandleAddLayerExecute,
+					TWeakObjectPtr<UCompositeLayerBase>(Layer), WeakSequencer));
+
+				MenuBuilder.AddMenuEntry(
+					DisplayName,
+					FText::Format(LOCTEXT("AddLayerTooltip", "Add {0} to Sequencer"), DisplayName),
+					FSlateIcon(),
+					Action);
+			}
+		}
+	}
+	MenuBuilder.EndSection();
+
+	// --- Components section (replicate FActorSchema behavior) ---
+	MenuBuilder.BeginSection("Components", LOCTEXT("ComponentsSection", "Components"));
+	{
+		struct FComponentEntry
+		{
+			FString DisplayString;
+			UActorComponent* Component;
+		};
+
+		TArray<FComponentEntry> ComponentEntries;
+
+		for (ACompositeActor* Actor : Actors)
+		{
+			for (UActorComponent* Component : Actor->GetComponents())
+			{
+				if (!Component || Component->IsVisualizationComponent())
+				{
+					continue;
+				}
+
+				if (Sequencer->GetHandleToObject(Component, false).IsValid())
+				{
+					continue;
+				}
+
+				ComponentEntries.Add({ Component->GetName(), Component });
+			}
+		}
+
+		ComponentEntries.Sort([](const FComponentEntry& A, const FComponentEntry& B)
+		{
+			return A.DisplayString < B.DisplayString;
+		});
+
+		for (const FComponentEntry& Entry : ComponentEntries)
+		{
+			FUIAction Action(FExecuteAction::CreateRaw(
+				this, &FCompositeActorObjectSchema::HandleAddComponentExecute,
+				TWeakObjectPtr<UActorComponent>(Entry.Component), WeakSequencer));
+
+			FText Label = FText::FromString(Entry.DisplayString);
+			MenuBuilder.AddMenuEntry(
+				Label,
+				FText::Format(LOCTEXT("AddComponentTooltip", "Add {0}"), Label),
+				FSlateIcon(),
+				Action);
+		}
+	}
+	MenuBuilder.EndSection();
+}
+
+void FCompositeActorObjectSchema::HandleLayerTrackMenu(
+	FMenuBuilder& MenuBuilder,
+	TWeakPtr<ISequencer> WeakSequencer,
+	TArray<UCompositeLayerBase*> Layers) const
+{
+	TSharedPtr<ISequencer> Sequencer = WeakSequencer.Pin();
+	if (!Sequencer)
+	{
+		return;
+	}
+
+	MenuBuilder.BeginSection("CompositePasses", LOCTEXT("CompositePassesSection", "Composite Passes"));
+	{
+		for (UCompositeLayerBase* Layer : Layers)
+		{
+			for (TFieldIterator<FArrayProperty> PropIt(Layer->GetClass()); PropIt; ++PropIt)
+			{
+				FArrayProperty* ArrayProp = *PropIt;
+				FObjectPropertyBase* InnerProp = CastField<FObjectPropertyBase>(ArrayProp->Inner);
+
+				if (!InnerProp || !InnerProp->PropertyClass->IsChildOf(UCompositePassBase::StaticClass()))
+				{
+					continue;
+				}
+
+				if (InnerProp->PropertyClass->IsChildOf(UCompositeLayerBase::StaticClass()))
+				{
+					continue;
+				}
+
+				FScriptArrayHelper ArrayHelper(ArrayProp, ArrayProp->ContainerPtrToValuePtr<void>(Layer));
+
+				for (int32 Index = 0; Index < ArrayHelper.Num(); ++Index)
+				{
+					UObject* Element = InnerProp->GetObjectPropertyValue(ArrayHelper.GetElementPtr(Index));
+					UCompositePassBase* Pass = Cast<UCompositePassBase>(Element);
+
+					if (!Pass || Pass->IsA<UCompositeLayerBase>())
+					{
+						continue;
+					}
+
+					if (Sequencer->GetHandleToObject(Pass, false).IsValid())
+					{
+						continue;
+					}
+
+					FText DisplayName;
+					const FString PassName = Pass->GetDisplayName();
+					if (!PassName.IsEmpty())
+					{
+						DisplayName = FText::FromString(PassName);
+					}
+					else
+					{
+						DisplayName = Pass->GetClass()->GetDisplayNameText();
+					}
+
+					FUIAction Action(FExecuteAction::CreateRaw(
+						this, &FCompositeActorObjectSchema::HandleAddPassExecute,
+						TWeakObjectPtr<UCompositePassBase>(Pass), WeakSequencer));
+
+					MenuBuilder.AddMenuEntry(
+						DisplayName,
+						FText::Format(LOCTEXT("AddPassTooltip", "Add {0} to Sequencer"), DisplayName),
+						FSlateIcon(),
+						Action);
+				}
+			}
+		}
+	}
+	MenuBuilder.EndSection();
+}
+
+void FCompositeActorObjectSchema::HandleAddLayerExecute(
+	TWeakObjectPtr<UCompositeLayerBase> Layer,
+	TWeakPtr<ISequencer> WeakSequencer) const
+{
+	TSharedPtr<ISequencer> Sequencer = WeakSequencer.Pin();
+	if (!Sequencer || !Layer.IsValid())
+	{
+		return;
+	}
+
+	const FScopedTransaction Transaction(LOCTEXT("AddCompositeLayer", "Add Composite Layer"));
+	Sequencer->GetHandleToObject(Layer.Get());
+}
+
+void FCompositeActorObjectSchema::HandleAddPassExecute(
+	TWeakObjectPtr<UCompositePassBase> Pass,
+	TWeakPtr<ISequencer> WeakSequencer) const
+{
+	TSharedPtr<ISequencer> Sequencer = WeakSequencer.Pin();
+	if (!Sequencer || !Pass.IsValid())
+	{
+		return;
+	}
+
+	const FScopedTransaction Transaction(LOCTEXT("AddCompositePass", "Add Composite Pass"));
+	Sequencer->GetHandleToObject(Pass.Get());
+}
+
+void FCompositeActorObjectSchema::HandleAddComponentExecute(
+	TWeakObjectPtr<UActorComponent> Component,
+	TWeakPtr<ISequencer> WeakSequencer) const
+{
+	TSharedPtr<ISequencer> Sequencer = WeakSequencer.Pin();
+	if (!Sequencer || !Component.IsValid())
+	{
+		return;
+	}
+
+	const FScopedTransaction Transaction(LOCTEXT("AddComponent", "Add Component"));
+	Sequencer->GetHandleToObject(Component.Get());
+}
+
+} // namespace UE::Sequencer
+
+#undef LOCTEXT_NAMESPACE
diff --git a/Engine/Plugins/Compositing/Composite/Source/CompositeEditor/Private/Sequencer/CompositeActorObjectSchema.h b/Engine/Plugins/Compositing/Composite/Source/CompositeEditor/Private/Sequencer/CompositeActorObjectSchema.h
new file mode 100644
--- /dev/null
+++ b/Engine/Plugins/Compositing/Composite/Source/CompositeEditor/Private/Sequencer/CompositeActorObjectSchema.h
@@ -0,0 +1,52 @@
+// Copyright Epic Games, Inc. All Rights Reserved.
+
+#pragma once
+
+#include "ISequencerObjectSchema.h"
+#include "UObject/WeakObjectPtrTemplates.h"
+
+class ACompositeActor;
+class FMenuBuilder;
+class UActorComponent;
+class UCompositeLayerBase;
+class UCompositePassBase;
+
+namespace UE::Sequencer
+{
+
+/**
+ * Object schema that exposes ACompositeActor layers and their passes
+ * as possessable sub-objects in the Sequencer outliner.
+ *
+ * Handles two nesting levels:
+ *   Actor -> Layers  (UCompositeLayerBase children of ACompositeActor)
+ *   Layer -> Passes  (UCompositePassBase children of UCompositeLayerBase)
+ */
+class FCompositeActorObjectSchema : public IObjectSchema
+{
+public:
+	//~ Begin IObjectSchema
+	virtual UObject* GetParentObject(UObject* Object) const override;
+	virtual FObjectSchemaRelevancy GetRelevancy(const UObject* InObject) const override;
+	virtual TSharedPtr<FExtender> ExtendObjectBindingMenu(TSharedRef<FUICommandList> CommandList, TWeakPtr<ISequencer> WeakSequencer, TArrayView<UObject* const> ContextSensitiveObjects) const override;
+	virtual FText GetPrettyName(const UObject* Object) const override;
+	//~ End IObjectSchema
+
+private:
+	/** Populates the actor-level track menu with composite layers and actor components not yet bound in Sequencer. */
+	void HandleActorTrackMenu(FMenuBuilder& MenuBuilder, TWeakPtr<ISequencer> WeakSequencer, TArray<ACompositeActor*> Actors) const;
+
+	/** Populates the layer-level track menu with composite passes (from all pass arrays) not yet bound in Sequencer. */
+	void HandleLayerTrackMenu(FMenuBuilder& MenuBuilder, TWeakPtr<ISequencer> WeakSequencer, TArray<UCompositeLayerBase*> Layers) const;
+
+	/** Creates a possessable binding for a composite layer in a scoped undo transaction. */
+	void HandleAddLayerExecute(TWeakObjectPtr<UCompositeLayerBase> Layer, TWeakPtr<ISequencer> WeakSequencer) const;
+
+	/** Creates a possessable binding for a composite pass in a scoped undo transaction. */
+	void HandleAddPassExecute(TWeakObjectPtr<UCompositePassBase> Pass, TWeakPtr<ISequencer> WeakSequencer) const;
+
+	/** Creates a possessable binding for an actor component in a scoped undo transaction. */
+	void HandleAddComponentExecute(TWeakObjectPtr<UActorComponent> Component, TWeakPtr<ISequencer> WeakSequencer) const;
+};
+
+} // namespace UE::Sequencer
diff --git a/Engine/Plugins/Compositing/Composite/Source/CompositeEditor/Public/CompositeEditorModule.h b/Engine/Plugins/Compositing/Composite/Source/CompositeEditor/Public/CompositeEditorModule.h
--- a/Engine/Plugins/Compositing/Composite/Source/CompositeEditor/Public/CompositeEditorModule.h
+++ b/Engine/Plugins/Compositing/Composite/Source/CompositeEditor/Public/CompositeEditorModule.h
@@ -7,6 +7,8 @@
 #include "Logging/LogMacros.h"
 #include "Modules/ModuleManager.h"
 
+namespace UE::Sequencer { class FCompositeActorObjectSchema; }
+
 /** Log category for this module. */
 DECLARE_LOG_CATEGORY_EXTERN(LogCompositeEditor, Log, All);
 
@@ -41,4 +43,7 @@ private:
 
 	/** Tab manager change delegate. */
 	FDelegateHandle OnTabManagerChangedDelegateHandle;
+
+	/** Sequencer object schema for exposing layers and passes as possessable sub-objects. */
+	TSharedPtr<UE::Sequencer::FCompositeActorObjectSchema> CompositeActorObjectSchema;
 };
