Composure: Remove actors without valid primitives from Composite Mesh Content picker on plate layer.

#Synced-CL 51392841

diff --git a/Engine/Plugins/Compositing/Composite/Source/CompositeEditor/Private/Customizations/CompositeLayerPlateCustomization.cpp b/Engine/Plugins/Compositing/Composite/Source/CompositeEditor/Private/Customizations/CompositeLayerPlateCustomization.cpp
--- a/Engine/Plugins/Compositing/Composite/Source/CompositeEditor/Private/Customizations/CompositeLayerPlateCustomization.cpp
+++ b/Engine/Plugins/Compositing/Composite/Source/CompositeEditor/Private/Customizations/CompositeLayerPlateCustomization.cpp
@@ -2,7 +2,9 @@
 
 #include "CompositeLayerPlateCustomization.h"
 
+#include "ActorTreeItem.h"
 #include "CompositeActor.h"
+#include "CompositeCoreSettings.h"
 #include "CompositeEditorStyle.h"
 #include "CompositeMeshActor.h"
 #include "DetailCategoryBuilder.h"
@@ -16,9 +18,11 @@
 #include "MediaSource.h"
 #include "MediaTexture.h"
 #include "PropertyCustomizationHelpers.h"
+#include "SceneOutlinerFilters.h"
 #include "ScopedTransaction.h"
 #include "UnrealEdGlobals.h"
 #include "Camera/CameraComponent.h"
+#include "Components/PrimitiveComponent.h"
 #include "Editor/UnrealEdEngine.h"
 #include "Framework/MultiBox/MultiBoxBuilder.h"
 #include "Layers/CompositeLayerPlate.h"
@@ -106,6 +110,29 @@ void FCompositeLayerPlateCustomization::CustomizeDetails(IDetailLayoutBuilder& D
 		.WholeRowContent()
 		[
 			SAssignNew(CompositeMeshPicker, SCompositeActorPickerTable, ActorListRef)
+			.SceneOutlinerFilters_Lambda([]()
+			{
+				TSharedPtr<FSceneOutlinerFilters> Filters = MakeShared<FSceneOutlinerFilters>();
+
+				Filters->AddFilterPredicate<FActorTreeItem>(FActorTreeItem::FFilterPredicate::CreateLambda([](const AActor* InActor)
+				{
+					const UCompositeCorePluginSettings* Settings = GetDefault<UCompositeCorePluginSettings>();
+
+					TInlineComponentArray<UPrimitiveComponent*> PrimComponents(InActor);
+					
+					for (const UPrimitiveComponent* PrimitiveComponent : PrimComponents)
+					{
+						if (Settings->IsAllowedPrimitiveClass(PrimitiveComponent))
+						{
+							return true;
+						}
+					}
+
+					return false;
+				}));
+
+				return Filters;
+			})
 			.OnExtendAddMenu(this, &FCompositeLayerPlateCustomization::OnExtendCompositeMeshAddMenu)
 			.OnLayoutSizeChanged(this, &FCompositeLayerPlateCustomization::OnLayoutSizeChanged)
 			.ShowApplyMaterialSection(true)
diff --git a/Engine/Plugins/Compositing/CompositeCore/Source/CompositeCore/Private/CompositeCoreSettings.cpp b/Engine/Plugins/Compositing/CompositeCore/Source/CompositeCore/Private/CompositeCoreSettings.cpp
--- a/Engine/Plugins/Compositing/CompositeCore/Source/CompositeCore/Private/CompositeCoreSettings.cpp
+++ b/Engine/Plugins/Compositing/CompositeCore/Source/CompositeCore/Private/CompositeCoreSettings.cpp
@@ -33,7 +33,8 @@ UCompositeCorePluginSettings::UCompositeCorePluginSettings()
 			FSoftClassPath("/Script/Engine.BillboardComponent"),
 			FSoftClassPath("/Script/Engine.ArrowComponent"),
 			FSoftClassPath("/Script/Engine.DrawFrustumComponent"),
-			FSoftClassPath("/Script/Engine.LineBatchComponent")
+			FSoftClassPath("/Script/Engine.LineBatchComponent"),
+			FSoftClassPath("/Script/Engine.DrawSphereComponent")
 		}
 	)
 	, AllowedComponentClasses(
