Composure: Fix undo regression caused by CL 50672860.

* Re-insert TryRegisterCompositePrimitives on composite mesh changes for transaction capture.
* Refactor visibility changes to separate activation from release, cleaning up reliance on GControlMeshVisibility state.

#Synced-CL 51392841

diff --git a/Engine/Plugins/Compositing/Composite/Source/Composite/Private/CompositeActor.cpp b/Engine/Plugins/Compositing/Composite/Source/Composite/Private/CompositeActor.cpp
--- a/Engine/Plugins/Compositing/Composite/Source/Composite/Private/CompositeActor.cpp
+++ b/Engine/Plugins/Compositing/Composite/Source/Composite/Private/CompositeActor.cpp
@@ -356,7 +356,7 @@ void ACompositeActor::PostRegisterAllComponents()
 
 	if (IsRendering())
 	{
-		PropagateStateChange(true);
+		PropagateStateChange(ECompositeStateChangeType::Activate);
 	}
 }
 
@@ -370,12 +370,9 @@ void ACompositeActor::UnregisterAllComponents(bool bForReregister)
 	{
 		return;
 	}
-	
-	// We don't want to hide meshes upon deletion
-	TGuardValue<bool> Guard(UCompositeLayerPlate::GControlMeshVisibility, false);
 
 	// State removal is done without any conditions to prevent stale state
-	PropagateStateChange(false);
+	PropagateStateChange(ECompositeStateChangeType::Release);
 }
 
 void ACompositeActor::PostUnregisterAllComponents()
@@ -616,7 +613,7 @@ void ACompositeActor::SetIsTemporarilyHiddenInEditor(bool bIsHidden)
 {
 	Super::SetIsTemporarilyHiddenInEditor(bIsHidden);
 
-	PropagateStateChange(IsRendering());
+	PropagateStateChange(IsRendering() ? ECompositeStateChangeType::Activate : ECompositeStateChangeType::Deactivate);
 }
 
 bool ACompositeActor::ShouldTickIfViewportsOnly() const
@@ -645,13 +642,13 @@ void ACompositeActor::SetActive(bool bInActive)
 				if (Actor && Actor != this)
 				{
 					Actor->bIsActive = false;
-					Actor->PropagateStateChange(false);
+					Actor->PropagateStateChange(ECompositeStateChangeType::Deactivate);
 				}
 			}
 		}
 	}
 
-	PropagateStateChange(IsRendering());
+	PropagateStateChange(IsRendering() ? ECompositeStateChangeType::Activate : ECompositeStateChangeType::Deactivate);
 
 	// Update the view projection matrix sent to the material parameter collection when changing which actor is active.
 	if (bIsActive && IsValid(ViewProjectionComponent))
@@ -669,7 +666,7 @@ void ACompositeActor::SetEnabled(bool bInEnabled)
 {
 	bIsEnabled = bInEnabled;
 
-	PropagateStateChange(IsRendering());
+	PropagateStateChange(IsRendering() ? ECompositeStateChangeType::Activate : ECompositeStateChangeType::Deactivate);
 }
 
 FComponentReference& ACompositeActor::GetCamera()
@@ -899,18 +896,18 @@ bool ACompositeActor::IsInVPRole() const
 	return !VPRole.IsValid() || (VPRolesSubsystem && VPRolesSubsystem->HasActiveRole(VPRole));
 }
 
-void ACompositeActor::PropagateStateChange(bool bApply)
+void ACompositeActor::PropagateStateChange(ECompositeStateChangeType ChangeType)
 {
 	for (TObjectPtr<UCompositeLayerBase>& Layer : CompositeLayers)
 	{
 		if (IsValid(Layer))
 		{
-			Layer->OnRenderingStateChange(bApply);
+			Layer->OnRenderingStateChange(ChangeType);
 		}
 	}
 
 	// Remove state case
-	if (!bApply)
+	if (ChangeType != ECompositeStateChangeType::Activate)
 	{
 		if (UCompositeWorldSubsystem* Subsystem = UWorld::GetSubsystem<UCompositeWorldSubsystem>(GetWorld()))
 		{
diff --git a/Engine/Plugins/Compositing/Composite/Source/Composite/Private/Layers/CompositeLayerPlate.cpp b/Engine/Plugins/Compositing/Composite/Source/Composite/Private/Layers/CompositeLayerPlate.cpp
--- a/Engine/Plugins/Compositing/Composite/Source/Composite/Private/Layers/CompositeLayerPlate.cpp
+++ b/Engine/Plugins/Compositing/Composite/Source/Composite/Private/Layers/CompositeLayerPlate.cpp
@@ -48,8 +48,18 @@ namespace UE
 	{
 		namespace Private
 		{
-			/** Convenience function to handle state on removed composite meshes. */
-			static void HandleRemovedCompositeMeshes(TConstArrayView<TSoftObjectPtr<AActor>> InPreEditMeshes, TConstArrayView<TSoftObjectPtr<AActor>> InPostEditMeshes, bool bRestoreVisibility)
+			/** Set visibility on a primitive component, guarded to avoid unnecessary Modify() calls. */
+			static void SetCompositeMeshVisibility(UPrimitiveComponent& InPrimitive, bool bNewVisibility)
+			{
+				if (InPrimitive.GetVisibleFlag() != bNewVisibility)
+				{
+					InPrimitive.Modify();
+					InPrimitive.SetVisibility(bNewVisibility);
+				}
+			}
+
+			/** Convenience function to handle state on removed composite meshes. Restores visibility when GControlMeshVisibility is enabled. */
+			static void HandleRemovedCompositeMeshes(TConstArrayView<TSoftObjectPtr<AActor>> InPreEditMeshes, TConstArrayView<TSoftObjectPtr<AActor>> InPostEditMeshes)
 			{
 				for (const TSoftObjectPtr<AActor>& PreEditMesh : InPreEditMeshes)
 				{
@@ -68,10 +78,9 @@ namespace UE
 							{
 								PrimitiveComponent->RemoveUserDataOfClass(UCompositeAssetUserData::StaticClass());
 
-								if (bRestoreVisibility)
+								if (UCompositeLayerPlate::GControlMeshVisibility)
 								{
-									constexpr bool bNewVisibility = true;
-									PrimitiveComponent->SetVisibility(true);
+									SetCompositeMeshVisibility(*PrimitiveComponent, true);
 								}
 							}
 						}
@@ -80,10 +89,26 @@ namespace UE
 			}
 
 			/** Convenience function to remove all composite asset user data from tracked mesh actors. */
-			static void RemoveCompositeAllAssetUserData(TConstArrayView<TSoftObjectPtr<AActor>> InPreEditMeshes)
+			static void RemoveCompositeAllAssetUserData(TConstArrayView<TSoftObjectPtr<AActor>> InMeshes)
 			{
-				constexpr bool bRestoreVisibility = false;
-				HandleRemovedCompositeMeshes(InPreEditMeshes, {}, bRestoreVisibility);
+				for (const TSoftObjectPtr<AActor>& Mesh : InMeshes)
+				{
+					if (!Mesh.IsValid())
+					{
+						continue;
+					}
+
+					for (UActorComponent* Component : Mesh->GetComponents())
+					{
+						if (UPrimitiveComponent* PrimitiveComponent = Cast<UPrimitiveComponent>(Component))
+						{
+							if (IsValid(PrimitiveComponent))
+							{
+								PrimitiveComponent->RemoveUserDataOfClass(UCompositeAssetUserData::StaticClass());
+							}
+						}
+					}
+				}
 			}
 
 			/** Return the texture size, depending on the texture type. */
@@ -143,8 +168,8 @@ UCompositeLayerPlate::~UCompositeLayerPlate() = default;
 void UCompositeLayerPlate::SetEnabled(bool bInEnabled)
 {
 	Super::SetEnabled(bInEnabled);
-	
-	PropagateStateChange(bInEnabled, GetWorld());
+
+	PropagateStateChange(bInEnabled ? ECompositeStateChangeType::Activate : ECompositeStateChangeType::Deactivate, GetWorld());
 }
 
 void UCompositeLayerPlate::Tick(float DeltaTime)
@@ -232,16 +257,12 @@ void UCompositeLayerPlate::OnRemoved(ACompositeActor* LastOwner)
 
 	TryCloseMediaProfileSource();
 
-	// We don't want to hide meshes upon deletion
-	TGuardValue<bool> Guard(GControlMeshVisibility, false);
-
-	constexpr bool bApply = false;
-	PropagateStateChange(bApply, LastOwner ? LastOwner->GetWorld() : nullptr);
+	PropagateStateChange(ECompositeStateChangeType::Release, LastOwner ? LastOwner->GetWorld() : nullptr);
 }
 
-void UCompositeLayerPlate::OnRenderingStateChange(bool bApply)
+void UCompositeLayerPlate::OnRenderingStateChange(ECompositeStateChangeType ChangeType)
 {
-	PropagateStateChange(bApply, GetWorld());
+	PropagateStateChange(ChangeType, GetWorld());
 }
 
 void UCompositeLayerPlate::BeginDestroy()
@@ -268,10 +289,10 @@ const TArray<TSoftObjectPtr<AActor>> UCompositeLayerPlate::GetCompositeMeshes()
 void UCompositeLayerPlate::SetCompositeMeshes(TArray<TSoftObjectPtr<AActor>> InCompositeMeshes)
 {
 	// First, we untrack current meshes
-	UE::Composite::Private::HandleRemovedCompositeMeshes(CompositeMeshes, InCompositeMeshes, GControlMeshVisibility);
+	UE::Composite::Private::HandleRemovedCompositeMeshes(CompositeMeshes, InCompositeMeshes);
 
 	CompositeMeshes = MoveTemp(InCompositeMeshes);
-	
+
 	UpdateCompositeMeshes();
 }
 
@@ -369,7 +390,7 @@ void UCompositeLayerPlate::PostEditChangeProperty(FPropertyChangedEvent& Propert
 
 		if (bChangeWasRemoval)
 		{
-			UE::Composite::Private::HandleRemovedCompositeMeshes(PreEditCompositeMeshes, CompositeMeshes, GControlMeshVisibility);
+			UE::Composite::Private::HandleRemovedCompositeMeshes(PreEditCompositeMeshes, CompositeMeshes);
 		}
 		else
 		{
@@ -394,6 +415,9 @@ void UCompositeLayerPlate::PostEditChangeProperty(FPropertyChangedEvent& Propert
 			}
 		}
 
+		// Note: While relying on tick would yield the same result, we want to capture the holdout state change in transactions.
+		TryRegisterCompositePrimitives(GetPrimitives());
+
 		PreEditCompositeMeshes.Reset();
 	}
 	else if (PropertyName == GET_MEMBER_NAME_CHECKED(ThisClass, Texture))
@@ -933,9 +957,9 @@ void UCompositeLayerPlate::UpdatePrimitiveComponent(UPrimitiveComponent& InPrimi
 	}
 }
 
-void UCompositeLayerPlate::PropagateStateChange(bool bApply, const UWorld* SourceWorld) const
+void UCompositeLayerPlate::PropagateStateChange(ECompositeStateChangeType ChangeType, const UWorld* SourceWorld) const
 {
-	if (bApply && IsRendering())
+	if (ChangeType == ECompositeStateChangeType::Activate && IsRendering())
 	{
 		if (CompositeMeshes.IsEmpty())
 		{
@@ -951,10 +975,7 @@ void UCompositeLayerPlate::PropagateStateChange(bool bApply, const UWorld* Sourc
 		{
 			for (UPrimitiveComponent* Primitive : Primitives)
 			{
-				Primitive->Modify();
-
-				constexpr bool bNewVisibility = true;
-				Primitive->SetVisibility(bNewVisibility);
+				UE::Composite::Private::SetCompositeMeshVisibility(*Primitive, true);
 			}
 		}
 
@@ -1000,17 +1021,24 @@ void UCompositeLayerPlate::PropagateStateChange(bool bApply, const UWorld* Sourc
 
 			if (GControlMeshVisibility)
 			{
-				TSet<UPrimitiveComponent*> UsedPrimitives = GetPrimitivesUsedByOtherActors();
-
-				for (UPrimitiveComponent* Primitive : Primitives)
+				if (ChangeType == ECompositeStateChangeType::Release)
 				{
-					// Only make a primitive invisible if it's not already used by another composite actor
-					if (!UsedPrimitives.Contains(Primitive))
+					for (UPrimitiveComponent* Primitive : Primitives)
 					{
-						Primitive->Modify();
+						UE::Composite::Private::SetCompositeMeshVisibility(*Primitive, true);
+					}
+				}
+				else
+				{
+					TSet<UPrimitiveComponent*> UsedPrimitives = GetPrimitivesUsedByOtherActors();
 
-						constexpr bool bNewVisibility = false;
-						Primitive->SetVisibility(bNewVisibility);
+					for (UPrimitiveComponent* Primitive : Primitives)
+					{
+						// Only make a primitive invisible if it's not already used by another composite actor
+						if (!UsedPrimitives.Contains(Primitive))
+						{
+							UE::Composite::Private::SetCompositeMeshVisibility(*Primitive, false);
+						}
 					}
 				}
 			}
diff --git a/Engine/Plugins/Compositing/Composite/Source/Composite/Public/CompositeActor.h b/Engine/Plugins/Compositing/Composite/Source/Composite/Public/CompositeActor.h
--- a/Engine/Plugins/Compositing/Composite/Source/Composite/Public/CompositeActor.h
+++ b/Engine/Plugins/Compositing/Composite/Source/Composite/Public/CompositeActor.h
@@ -214,12 +214,8 @@ public:
 	UE_API bool IsInVPRole() const;
 
 private:
-	/**
-	* Propagate rendering state changes to layers and the composite core subsystem.
-	* 
-	* @param bApply - Apply or remove external rendering state.
-	*/
-	void PropagateStateChange(bool bApply);
+	/** Propagate rendering state changes to layers and the composite core subsystem. */
+	void PropagateStateChange(ECompositeStateChangeType ChangeType);
 
 	/** Layer-managed scene capture components. */
 	UPROPERTY()
diff --git a/Engine/Plugins/Compositing/Composite/Source/Composite/Public/Layers/CompositeLayerBase.h b/Engine/Plugins/Compositing/Composite/Source/Composite/Public/Layers/CompositeLayerBase.h
--- a/Engine/Plugins/Compositing/Composite/Source/Composite/Public/Layers/CompositeLayerBase.h
+++ b/Engine/Plugins/Compositing/Composite/Source/Composite/Public/Layers/CompositeLayerBase.h
@@ -15,6 +15,17 @@
 class ACompositeActor;
 class UWorld;
 
+/** Describes the intent of a rendering state change propagated to layers. */
+enum class ECompositeStateChangeType : uint8
+{
+	/** Show meshes, register primitives (activate rendering). */
+	Activate,
+	/** Hide meshes, unregister primitives (deactivate rendering). */
+	Deactivate,
+	/** Restore mesh visibility, unregister primitives (removal/destruction cleanup). */
+	Release
+};
+
 /** Base class for defining a composite layer object, and its associated render proxy. */
 UCLASS(MinimalAPI, Abstract, EditInlineNew)
 class UCompositeLayerBase : public UCompositePassBase
@@ -44,12 +55,11 @@ public:
 	*/
 	UE_API virtual void OnEndOfFrameUpdate(UWorld* InWorld) {}
 
-	/**
-	* Called when composite actor rendering state changes.
+	/** Called when composite actor rendering state changes.
 	* 
-	* @param bApply - Apply or remove external rendering state.
+	* @param ChangeType - Activation, deactivation or release state.
 	*/
-	UE_API virtual void OnRenderingStateChange(bool bApply) {}
+	UE_API virtual void OnRenderingStateChange(ECompositeStateChangeType ChangeType) {}
 
 	/**
 	* Called when the layer is removed or composite actor is destroyed.
diff --git a/Engine/Plugins/Compositing/Composite/Source/Composite/Public/Layers/CompositeLayerPlate.h b/Engine/Plugins/Compositing/Composite/Source/Composite/Public/Layers/CompositeLayerPlate.h
--- a/Engine/Plugins/Compositing/Composite/Source/Composite/Public/Layers/CompositeLayerPlate.h
+++ b/Engine/Plugins/Compositing/Composite/Source/Composite/Public/Layers/CompositeLayerPlate.h
@@ -51,7 +51,7 @@ public:
 
 	UE_API virtual void OnRemoved(ACompositeActor* LastOwner) override;
 
-	UE_API virtual void OnRenderingStateChange(bool bApply) override;
+	UE_API virtual void OnRenderingStateChange(ECompositeStateChangeType ChangeType) override;
 
 	/** Set the enabled state. */
 	UE_API virtual void SetEnabled(bool bInEnabled) override;
@@ -182,13 +182,8 @@ private:
 	/** Update the (media) texture on a primitive component. */
 	void UpdatePrimitiveComponent(UPrimitiveComponent& InPrimitiveComponent) const;
 
-	/**
-	* Apply (or remove) state to activate (or deactivate) the built-in custom render pass for composite meshes.
-	* 
-	* @param bApply Apply or remove external rendering state.
-	* @param SourceWorld Optional source world to find the composite subsystem
-	*/
-	void PropagateStateChange(bool bApply, const UWorld* SourceWorld) const;
+	/** Propagate rendering state change to composite meshes. */
+	void PropagateStateChange(ECompositeStateChangeType ChangeType, const UWorld* SourceWorld) const;
 
 	/** Convenience function verify if the composite mesh actor is already in use by another layer. */
 	bool IsCompositeMeshActorAlreadyInUse(TSoftObjectPtr<AActor> InCompositeMeshActor) const;
