Composure: Fix keyer devignette regression after engine CL 49896456.

The lerp(Color, float3(1), Weight) formula with white Color would always returns 1.0 regardless of Weight, so VignetteMask was always 1 -> rcp(1.0), hence no devignette.

#Synced-CL 51392841

diff --git a/Engine/Plugins/Compositing/Composite/Shaders/Private/CompositeColorKeyerCommon.ush b/Engine/Plugins/Compositing/Composite/Shaders/Private/CompositeColorKeyerCommon.ush
--- a/Engine/Plugins/Compositing/Composite/Shaders/Private/CompositeColorKeyerCommon.ush
+++ b/Engine/Plugins/Compositing/Composite/Shaders/Private/CompositeColorKeyerCommon.ush
@@ -161,13 +161,15 @@ float3 BlueDespillAvg(float3 Color)
 float Vignette(float2 ScreenPos, float AspectRatio, float Intensity)
 {
 	// Scale vignette to always be a circle with consistent corner intensity.
-	float2 Vignette = VignetteSpace(ScreenPos, VIGNETTE_TYPE_COSINE_FOURTH_LAW, AspectRatio);
+	// We compute the raw cosine-fourth-law scalar directly rather than using
+	// ComputeVignetteMask(), which now returns a float3 color tint intended for
+	// the post-process stack. The devignette path needs a plain 0-1 scalar to
+	// divide out the lens falloff.
+	float2 VignetteCircleSpacePos = VignetteSpace(ScreenPos, VIGNETTE_TYPE_COSINE_FOURTH_LAW, AspectRatio);
+	VignetteCircleSpacePos *= Intensity;
+	float Tan2Angle = dot( VignetteCircleSpacePos, VignetteCircleSpacePos );
 	
-	float4 VignetteParams0, VignetteParams1;
-	float VignetteParams2;
-	PackVignetteParameters(VignetteParams0, VignetteParams1, VignetteParams2, Intensity);
-
-	return ComputeVignetteMask(Vignette, VIGNETTE_TYPE_COSINE_FOURTH_LAW, VignetteParams0, VignetteParams1, VignetteParams2);
+	return Square( rcp( Tan2Angle + 1 ) );
 }
 
 float4 GreenScreenKeyer(
