Composure: Added fade parameter on alpha transform pass.

Removal/cleanup of implementation only properties, inverted alpha management of passes to follow.

#Synced-CL 51711852

diff --git a/Engine/Plugins/Compositing/Composite/Shaders/Private/CompositeAlphaTransform.usf b/Engine/Plugins/Compositing/Composite/Shaders/Private/CompositeAlphaTransform.usf
--- a/Engine/Plugins/Compositing/Composite/Shaders/Private/CompositeAlphaTransform.usf
+++ b/Engine/Plugins/Compositing/Composite/Shaders/Private/CompositeAlphaTransform.usf
@@ -10,7 +10,7 @@ Texture2D Input_Texture;
 SamplerState Input_Sampler;
 
 uint PremultOp;
-uint bInvertSourceAlpha;
+float Fade;
 uint OverrideOutputAlpha;
 
 #define PremultOp_None 0u
@@ -29,15 +29,19 @@ void MainPS(
     float2 TextureUV = ApplyScreenTransform(SvPosition.xy, SvPositionToInputTextureUV);
     float4 Color = Input_Texture.SampleLevel(Input_Sampler, TextureUV, 0);
 
-	Color.a = select(bInvertSourceAlpha > 0, saturate(1.0 - Color.a), Color.a);
-	
 	if(PremultOp == PremultOp_Premultiply)
 	{
 		Color.rgb *= Color.a;
+		Color *= Fade;
 	}
 	else if(PremultOp == PremultOp_Unpremultiply)
 	{
 		Color.rgb *= SafeRcp(Color.a);
+		Color.a *= Fade;
+	}
+	else
+	{
+		Color *= Fade;
 	}
 
 	if(OverrideOutputAlpha == OverrideOutput_Zero)
diff --git a/Engine/Plugins/Compositing/Composite/Source/Composite/Private/Passes/CompositePassAlphaTransform.cpp b/Engine/Plugins/Compositing/Composite/Source/Composite/Private/Passes/CompositePassAlphaTransform.cpp
--- a/Engine/Plugins/Compositing/Composite/Source/Composite/Private/Passes/CompositePassAlphaTransform.cpp
+++ b/Engine/Plugins/Compositing/Composite/Source/Composite/Private/Passes/CompositePassAlphaTransform.cpp
@@ -20,7 +20,7 @@ class FCompositePassAlphaTransformShader : public FGlobalShader
 		SHADER_PARAMETER_STRUCT(FScreenPassTextureInput, Input)
 		SHADER_PARAMETER(FScreenTransform, SvPositionToInputTextureUV)
 		SHADER_PARAMETER(uint32, PremultOp)
-		SHADER_PARAMETER(uint32, bInvertSourceAlpha)
+		SHADER_PARAMETER(float, Fade)
 		SHADER_PARAMETER(uint32, OverrideOutputAlpha)
 		RENDER_TARGET_BINDING_SLOTS()
 	END_SHADER_PARAMETER_STRUCT()
@@ -57,7 +57,7 @@ namespace UE
 					FScreenTransform::ChangeTextureBasisFromTo(FScreenPassTextureViewport(Output), FScreenTransform::ETextureBasis::TexelPosition, FScreenTransform::ETextureBasis::ViewportUV) *
 					FScreenTransform::ChangeTextureBasisFromTo(FScreenPassTextureViewport(InputTexture), FScreenTransform::ETextureBasis::ViewportUV, FScreenTransform::ETextureBasis::TextureUV));
 				Parameters->PremultOp = static_cast<uint32>(PremultOp);
-				Parameters->bInvertSourceAlpha = static_cast<uint32>(bInvertSourceAlpha);
+				Parameters->Fade = Fade;
 				Parameters->OverrideOutputAlpha = static_cast<uint32>(OverrideOutputAlpha);
 				Parameters->RenderTargets[0] = Output.GetRenderTargetBinding();
 
@@ -80,7 +80,7 @@ namespace UE
 UCompositePassAlphaTransform::UCompositePassAlphaTransform(const FObjectInitializer& ObjectInitializer)
 	: Super(ObjectInitializer)
 	, PremultOp{ ECompositeAlphaPremultiplication::None }
-	, bInvertSourceAlpha{ false }
+	, Fade{ 1.0f }
 	, OverrideOutputAlpha{ ECompositeAlphaOverride::None }
 {
 }
@@ -93,7 +93,7 @@ bool UCompositePassAlphaTransform::GetProxy(const UE::CompositeCore::FPassInputD
 
 	FAlphaTransformPassProxy* Proxy = InFrameAllocator.Create<FAlphaTransformPassProxy>(FPassInputDeclArray{ InputDecl });
 	Proxy->PremultOp = PremultOp;
-	Proxy->bInvertSourceAlpha = bInvertSourceAlpha;
+	Proxy->Fade = Fade;
 	Proxy->OverrideOutputAlpha = OverrideOutputAlpha;
 
 	OutProxy = Proxy;
diff --git a/Engine/Plugins/Compositing/Composite/Source/Composite/Public/Passes/CompositePassAlphaTransform.h b/Engine/Plugins/Compositing/Composite/Source/Composite/Public/Passes/CompositePassAlphaTransform.h
--- a/Engine/Plugins/Compositing/Composite/Source/Composite/Public/Passes/CompositePassAlphaTransform.h
+++ b/Engine/Plugins/Compositing/Composite/Source/Composite/Public/Passes/CompositePassAlphaTransform.h
@@ -46,12 +46,11 @@ public:
 	UPROPERTY(BlueprintReadWrite, EditAnywhere, Interp, Category = "Composite")
 	ECompositeAlphaPremultiplication PremultOp;
 
-	/** Invert the source alpha, for instance if it comes from a scene capture render. */
-	UPROPERTY(BlueprintReadWrite, EditAnywhere, Interp, Category = "Composite", AdvancedDisplay)
-	bool bInvertSourceAlpha;
+	/** Fade the alpha from 0 to 1. Assumes pre-multiplied alpha, so both RGB and A are scaled together. */
+	UPROPERTY(BlueprintReadWrite, EditAnywhere, Interp, Category = "Composite", meta = (ClampMin = "0.0", ClampMax = "1.0", UIMin = "0.0", UIMax = "1.0"))
+	float Fade = 1.0f;
 
-	/** Output alpha is forced to 1.0, primarily used for correct blending with composite mesh holdouts. */
-	UPROPERTY(BlueprintReadWrite, EditAnywhere, Interp, Category = "Composite", AdvancedDisplay)
+	UPROPERTY()
 	ECompositeAlphaOverride OverrideOutputAlpha;
 };
 
@@ -73,7 +72,7 @@ namespace UE
 				FPassTexture Add(FRDGBuilder& GraphBuilder, const FSceneView& InView, const FPassInputArray& Inputs, const FPassContext& PassContext) const override;
 
 				ECompositeAlphaPremultiplication PremultOp = ECompositeAlphaPremultiplication::None;
-				bool bInvertSourceAlpha = false;
+				float Fade = 1.0f;
 				ECompositeAlphaOverride OverrideOutputAlpha = ECompositeAlphaOverride::None;
 			};
 		}
