Composure: Fix header checkbox behavior for passes and (composite mesh) actor picker outliner.

#Synced-CL 51753773

diff --git a/Engine/Plugins/Compositing/Composite/Source/CompositeEditor/Private/UI/SCompositeActorPickerSceneOutliner.cpp b/Engine/Plugins/Compositing/Composite/Source/CompositeEditor/Private/UI/SCompositeActorPickerSceneOutliner.cpp
--- a/Engine/Plugins/Compositing/Composite/Source/CompositeEditor/Private/UI/SCompositeActorPickerSceneOutliner.cpp
+++ b/Engine/Plugins/Compositing/Composite/Source/CompositeEditor/Private/UI/SCompositeActorPickerSceneOutliner.cpp
@@ -14,6 +14,7 @@
 #include "Modules/ModuleManager.h"
 #include "Widgets/Input/SCheckBox.h"
 #include "Widgets/Layout/SBox.h"
+#include "Widgets/Views/STreeView.h"
 
 #define LOCTEXT_NAMESPACE "SCompositeActorPickerTable"
 
@@ -23,6 +24,7 @@ class FCompositeActorPickerSelectedColumn : public ISceneOutlinerColumn
 public:
 	DECLARE_DELEGATE_RetVal_OneParam(bool, FOnGetActorSelected, AActor*);
 	DECLARE_DELEGATE_TwoParams(FOnActorSelected, AActor*, bool /* bIsSelected */);
+	DECLARE_DELEGATE_TwoParams(FOnAllActorsSelected, const TArray<AActor*>& /* InActors */, bool /* bIsSelected */);
 	
 public:
 	FCompositeActorPickerSelectedColumn(ISceneOutliner& InSceneOutliner)
@@ -49,7 +51,8 @@ public:
 			.VAlignCell(VAlign_Center)
 			[
 				SNew(SCheckBox)
-				.IsChecked_Lambda([] { return ECheckBoxState::Checked; })
+				.IsChecked(this, &FCompositeActorPickerSelectedColumn::GetHeaderCheckState)
+				.OnCheckStateChanged(this, &FCompositeActorPickerSelectedColumn::OnHeaderCheckStateChangedHandler)
 			];
 	}
 
@@ -67,6 +70,102 @@ public:
 	// End of ISceneOutlinerColumn interface
 	
 private:
+	/** Recursively visits all actor tree items in the scene outliner, calling the visitor for each valid actor. */
+	void ForEachOutlinerActor(TFunctionRef<void(AActor*)> Visitor) const
+	{
+		TSharedPtr<ISceneOutliner> Outliner = SceneOutliner.Pin();
+		if (!Outliner.IsValid())
+		{
+			return;
+		}
+
+		// Recursive lambda to walk the tree depth-first
+		TFunction<void(const ISceneOutlinerTreeItem&)> VisitItem = [&](const ISceneOutlinerTreeItem& InItem)
+		{
+			if (const FActorTreeItem* ActorTreeItem = InItem.CastTo<FActorTreeItem>())
+			{
+				if (TStrongObjectPtr<AActor> PinnedActor = ActorTreeItem->Actor.Pin())
+				{
+					Visitor(PinnedActor.Get());
+				}
+			}
+
+			for (const TWeakPtr<ISceneOutlinerTreeItem>& ChildWeak : InItem.GetChildren())
+			{
+				if (TSharedPtr<ISceneOutlinerTreeItem> Child = ChildWeak.Pin())
+				{
+					VisitItem(*Child);
+				}
+			}
+		};
+
+		for (const FSceneOutlinerTreeItemPtr& RootItem : Outliner->GetTree().GetRootItems())
+		{
+			if (RootItem.IsValid())
+			{
+				VisitItem(*RootItem);
+			}
+		}
+	}
+
+	/** Returns the check state for the header checkbox based on all actors in the outliner: Checked if all are selected, Undetermined if mixed, Unchecked if none */
+	ECheckBoxState GetHeaderCheckState() const
+	{
+		if (!OnGetActorSelected.IsBound())
+		{
+			return ECheckBoxState::Unchecked;
+		}
+
+		bool bAnySelected = false;
+		bool bAnyUnselected = false;
+
+		ForEachOutlinerActor([&](AActor* Actor)
+		{
+			if (OnGetActorSelected.Execute(Actor))
+			{
+				bAnySelected = true;
+			}
+			else
+			{
+				bAnyUnselected = true;
+			}
+		});
+
+		if (bAnySelected)
+		{
+			return bAnyUnselected ? ECheckBoxState::Undetermined : ECheckBoxState::Checked;
+		}
+
+		return ECheckBoxState::Unchecked;
+	}
+
+	/** Toggles all actors in the outliner: selects all if any are unselected, deselects all if all are already selected */
+	void OnHeaderCheckStateChangedHandler(ECheckBoxState InNewCheckState)
+	{
+		if (!OnAllActorsSelected.IsBound() || !OnGetActorSelected.IsBound())
+		{
+			return;
+		}
+
+		// Single pass: collect all actors and determine if we should select or deselect
+		TArray<AActor*> Actors;
+		bool bAllSelected = true;
+
+		ForEachOutlinerActor([&](AActor* Actor)
+		{
+			Actors.Add(Actor);
+			if (!OnGetActorSelected.Execute(Actor))
+			{
+				bAllSelected = false;
+			}
+		});
+
+		if (!Actors.IsEmpty())
+		{
+			OnAllActorsSelected.Execute(Actors, !bAllSelected);
+		}
+	}
+
 	ECheckBoxState GetTreeItemSelectedCheckState(FSceneOutlinerTreeItemRef InTreeItem) const
 	{
 		if (FActorTreeItem* ActorTreeItem = InTreeItem->CastTo<FActorTreeItem>())
@@ -104,6 +203,7 @@ private:
 public:
 	FOnGetActorSelected OnGetActorSelected;
 	FOnActorSelected OnActorSelected;
+	FOnAllActorsSelected OnAllActorsSelected;
 };
 
 void SCompositeActorPickerSceneOutliner::Construct(const FArguments& InArgs, const FCompositeActorPickerListRef& InActorListRef)
@@ -129,6 +229,7 @@ void SCompositeActorPickerSceneOutliner::Construct(const FArguments& InArgs, con
 				TSharedPtr<FCompositeActorPickerSelectedColumn> Column = MakeShared<FCompositeActorPickerSelectedColumn>(Outliner);
 				Column->OnGetActorSelected.BindSP(this, &SCompositeActorPickerSceneOutliner::OnGetActorSelected);
 				Column->OnActorSelected.BindSP(this, &SCompositeActorPickerSceneOutliner::OnActorSelected);
+				Column->OnAllActorsSelected.BindSP(this, &SCompositeActorPickerSceneOutliner::OnAllActorsSelected);
 				
 				return Column.ToSharedRef();
 			})
@@ -206,4 +307,54 @@ void SCompositeActorPickerSceneOutliner::OnActorSelected(AActor* InActor, bool b
 	OnActorListChanged.ExecuteIfBound();
 }
 
+void SCompositeActorPickerSceneOutliner::OnAllActorsSelected(const TArray<AActor*>& InActors, bool bIsSelected)
+{
+	if (!ActorListRef.IsValid())
+	{
+		return;
+	}
+
+	TStrongObjectPtr<UObject> PinnedListOwner = ActorListRef.ActorListOwner.Pin();
+
+	FScopedTransaction Transaction(bIsSelected
+		? LOCTEXT("SelectAllActorsTransaction", "Select All Actors")
+		: LOCTEXT("DeselectAllActorsTransaction", "Deselect All Actors"));
+	PinnedListOwner->Modify();
+
+	ActorListRef.NotifyPreEditChange();
+
+	bool bModified = false;
+
+	for (AActor* Actor : InActors)
+	{
+		if (!Actor)
+		{
+			continue;
+		}
+
+		const bool bIsInList = ActorListRef.ActorList->Contains(Actor);
+
+		if (bIsSelected && !bIsInList)
+		{
+			ActorListRef.ActorList->Add(Actor);
+			bModified = true;
+		}
+		else if (!bIsSelected && bIsInList)
+		{
+			ActorListRef.ActorList->Remove(Actor);
+			bModified = true;
+		}
+	}
+
+	if (bModified)
+	{
+		ActorListRef.NotifyPostEditChangeList(bIsSelected ? EPropertyChangeType::ArrayAdd : EPropertyChangeType::ArrayRemove);
+		OnActorListChanged.ExecuteIfBound();
+	}
+	else
+	{
+		Transaction.Cancel();
+	}
+}
+
 #undef LOCTEXT_NAMESPACE
\ No newline at end of file
diff --git a/Engine/Plugins/Compositing/Composite/Source/CompositeEditor/Private/UI/SCompositeActorPickerSceneOutliner.h b/Engine/Plugins/Compositing/Composite/Source/CompositeEditor/Private/UI/SCompositeActorPickerSceneOutliner.h
--- a/Engine/Plugins/Compositing/Composite/Source/CompositeEditor/Private/UI/SCompositeActorPickerSceneOutliner.h
+++ b/Engine/Plugins/Compositing/Composite/Source/CompositeEditor/Private/UI/SCompositeActorPickerSceneOutliner.h
@@ -26,6 +26,9 @@ private:
 
 	/** Raised when the selection state of a specific actor has changed */
 	void OnActorSelected(AActor* InActor, bool bIsSelected);
+
+	/** Raised when the header checkbox toggles to select or deselect all actors at once */
+	void OnAllActorsSelected(const TArray<AActor*>& InActors, bool bIsSelected);
 	
 private:
 	/** Reference to the actor list being modified by this actor picker */
diff --git a/Engine/Plugins/Compositing/Composite/Source/CompositeEditor/Private/UI/SCompositePassTree.cpp b/Engine/Plugins/Compositing/Composite/Source/CompositeEditor/Private/UI/SCompositePassTree.cpp
--- a/Engine/Plugins/Compositing/Composite/Source/CompositeEditor/Private/UI/SCompositePassTree.cpp
+++ b/Engine/Plugins/Compositing/Composite/Source/CompositeEditor/Private/UI/SCompositePassTree.cpp
@@ -583,6 +583,7 @@ private:
 					SNew(SCheckBox)
 					.IsChecked(this, &SCompositePassTreeItemRow::IsItemEnabled)
 					.OnCheckStateChanged(this, &SCompositePassTreeItemRow::OnIsItemEnabledChanged)
+					.IsEnabled(this, &SCompositePassTreeItemRow::CanToggleEnabled)
 				];
 		}
 		if (InColumnName == CompositePassTree::PassTreeColumn_Name)
@@ -680,10 +681,26 @@ private:
 		{
 			return bAnyDisabled ? ECheckBoxState::Undetermined : ECheckBoxState::Checked;
 		}
-		else
+
+		return ECheckBoxState::Unchecked;
+	}
+
+	/** Gets whether the row item's Enabled checkbox can be toggled */
+	bool CanToggleEnabled() const
+	{
+		if (!TreeItem->HasValidOwner())
+		{
+			return false;
+		}
+
+		// Individual pass items can always be toggled
+		if (TreeItem->HasValidPassIndex())
 		{
-			return bAnyDisabled ? ECheckBoxState::Unchecked : ECheckBoxState::Checked;
+			return true;
 		}
+
+		// Group items can only be toggled when they have children
+		return !TreeItem->Children.IsEmpty();
 	}
 
 	/** Raised when the row item's Enabled checkbox has changed state */
@@ -1026,6 +1043,7 @@ void SCompositePassTree::Construct(const FArguments& InArgs, const TSharedPtr<IC
 					SNew(SCheckBox)
 					.IsChecked(this, &SCompositePassTree::GetGlobalEnabledState)
 					.OnCheckStateChanged(this, &SCompositePassTree::OnGlobalEnabledStateChanged)
+					.IsEnabled(this, &SCompositePassTree::HasAnyPasses)
 				]
 
 				+SHeaderRow::Column(CompositePassTree::PassTreeColumn_Name)
@@ -1453,15 +1471,9 @@ ECheckBoxState SCompositePassTree::GetGlobalEnabledState() const
 	{
 		const TArray<TObjectPtr<UCompositePassBase>>& PassList = PassListOwner->GetPassesForGroup(InGroupIndex);
 
-		// Consider any pass list that is empty as enabled to align with how the check state of the pass type tree item works
-		if (PassList.IsEmpty())
-		{
-			bAnyPassesEnabled = true;
-		}
-		
 		for (const TObjectPtr<UCompositePassBase>& Pass : PassList)
 		{
-			if (!IsValid(Pass))
+			if (!Pass)
 			{
 				continue;
 			}
@@ -1494,10 +1506,32 @@ ECheckBoxState SCompositePassTree::GetGlobalEnabledState() const
 	{
 		return bAnyPassesDisabled ? ECheckBoxState::Undetermined : ECheckBoxState::Checked;
 	}
-	else
+
+	return ECheckBoxState::Unchecked;
+}
+
+bool SCompositePassTree::HasAnyPasses() const
+{
+	if (!PassListOwner.IsValid() || !PassListOwner->IsObjectValid())
+	{
+		return false;
+	}
+
+	const int32 NumGroups = PassListOwner->GetNumGroups();
+	if (NumGroups <= 0)
+	{
+		return !PassListOwner->GetPassesForGroup(INDEX_NONE).IsEmpty();
+	}
+
+	for (int32 GroupIndex = 0; GroupIndex < NumGroups; ++GroupIndex)
 	{
-		return bAnyPassesDisabled ? ECheckBoxState::Unchecked : ECheckBoxState::Undetermined;
+		if (!PassListOwner->GetPassesForGroup(GroupIndex).IsEmpty())
+		{
+			return true;
+		}
 	}
+
+	return false;
 }
 
 void SCompositePassTree::OnGlobalEnabledStateChanged(ECheckBoxState CheckBoxState)
@@ -1521,6 +1555,11 @@ void SCompositePassTree::OnGlobalEnabledStateChanged(ECheckBoxState CheckBoxStat
 
 		for (const TObjectPtr<UCompositePassBase>& Pass : PassList)
 		{
+			if (!Pass)
+			{
+				continue;
+			}
+
 			if (!AllEnabledTransaction.IsValid())
 			{
 				AllEnabledTransaction = MakeShared<FScopedTransaction>(LOCTEXT("AllEnabledTransaction", "Set All Enabled"));
diff --git a/Engine/Plugins/Compositing/Composite/Source/CompositeEditor/Private/UI/SCompositePassTree.h b/Engine/Plugins/Compositing/Composite/Source/CompositeEditor/Private/UI/SCompositePassTree.h
--- a/Engine/Plugins/Compositing/Composite/Source/CompositeEditor/Private/UI/SCompositePassTree.h
+++ b/Engine/Plugins/Compositing/Composite/Source/CompositeEditor/Private/UI/SCompositePassTree.h
@@ -165,6 +165,9 @@ private:
 
 	/** Sets the enabled state for all passes on the plate */
 	void OnGlobalEnabledStateChanged(ECheckBoxState CheckBoxState);
+
+	/** Gets whether any passes exist across all groups */
+	bool HasAnyPasses() const;
 	
 	/** Raised when a pass has been moved via a drag drop operation */
 	void OnPassMoved(const TSharedPtr<FPassTreeItem>& InTreeItem, int32 InDestGroupIndex, int InDestIndex);
