Composure: Replaced camera component reference with actor soft pointer, such that cameras on sublevels can be selected/serialized.

#Synced-CL 51806324

diff --git a/Engine/Plugins/Compositing/Composite/Source/Composite/Private/Components/CompositeSceneCapture2DComponent.cpp b/Engine/Plugins/Compositing/Composite/Source/Composite/Private/Components/CompositeSceneCapture2DComponent.cpp
--- a/Engine/Plugins/Compositing/Composite/Source/Composite/Private/Components/CompositeSceneCapture2DComponent.cpp
+++ b/Engine/Plugins/Compositing/Composite/Source/Composite/Private/Components/CompositeSceneCapture2DComponent.cpp
@@ -77,10 +77,7 @@ const AActor* UCompositeSceneCapture2DComponent::GetViewOwner() const
 	ACompositeActor* CompositeActor = GetTypedOuter<ACompositeActor>();
 	if (IsValid(CompositeActor))
 	{
-		if (UCameraComponent* CameraComponent = Cast<UCameraComponent>(CompositeActor->GetCamera().GetComponent(nullptr)))
-		{
-			return CameraComponent->GetOwner();
-		}
+		return CompositeActor->GetCameraActor().Get();
 	}
 
 	return nullptr;
diff --git a/Engine/Plugins/Compositing/Composite/Source/Composite/Private/Components/CompositeViewProjectionComponent.cpp b/Engine/Plugins/Compositing/Composite/Source/Composite/Private/Components/CompositeViewProjectionComponent.cpp
--- a/Engine/Plugins/Compositing/Composite/Source/Composite/Private/Components/CompositeViewProjectionComponent.cpp
+++ b/Engine/Plugins/Compositing/Composite/Source/Composite/Private/Components/CompositeViewProjectionComponent.cpp
@@ -11,6 +11,7 @@
 #include "Materials/MaterialParameterCollection.h"
 #include "Materials/MaterialParameterCollectionInstance.h"
 #include "UObject/ConstructorHelpers.h"
+#include "UObject/UE5MainStreamObjectVersion.h"
 
 #define LOCTEXT_NAMESPACE "Composite"
 
@@ -154,15 +155,31 @@ void UCompositeViewProjectionComponent::ForceUpdate()
 	LastViewProjectionMatrix = FMatrix::Identity;
 }
 
-FComponentReference& UCompositeViewProjectionComponent::GetTargetComponent()
+UCameraComponent* UCompositeViewProjectionComponent::GetCameraComponent() const
 {
-	return TargetCameraComponent;
+	if (CameraActor.IsValid())
+	{
+		return CameraActor->GetComponentByClass<UCameraComponent>();
+	}
+
+	return nullptr;
+}
+
+void UCompositeViewProjectionComponent::SetCameraComponent(UCameraComponent* InComponent)
+{
+	CameraActor = IsValid(InComponent) ? InComponent->GetOwner() : nullptr;
+
+	ForceUpdate();
 }
 
-void UCompositeViewProjectionComponent::SetTargetComponent(const FComponentReference& InComponentReference)
+const TSoftObjectPtr<AActor>& UCompositeViewProjectionComponent::GetCameraActor()
 {
-	TargetCameraComponent = InComponentReference;
-	LastCameraComponent = Cast<UCameraComponent>(TargetCameraComponent.GetComponent(nullptr));
+	return CameraActor;
+}
+
+void UCompositeViewProjectionComponent::SetCameraActor(const TSoftObjectPtr<AActor>& InActor)
+{
+	CameraActor = InActor;
 
 	ForceUpdate();
 }
@@ -170,31 +187,15 @@ void UCompositeViewProjectionComponent::SetTargetComponent(const FComponentRefer
 void UCompositeViewProjectionComponent::InitDefaultCamera()
 {
 	ACompositeActor* CompositeActor = GetTypedOuter<ACompositeActor>();
+	
 	if (IsValid(CompositeActor))
 	{
-		// Special case when the component is owned by the composite actor: we use its camera component reference directly.
-		SetTargetComponent(CompositeActor->GetCamera());
-
-		return;
+		// Special case when the component is owned by the composite actor: we use its camera reference directly.
+		SetCameraActor(CompositeActor->GetCameraActor());
 	}
-
-	UCameraComponent* TargetCamera = Cast<UCameraComponent>(TargetCameraComponent.GetComponent(GetOwner()));
-	
-	if (!IsValid(TargetCamera))
+	else
 	{
-		// Find all CineCameraComponents on the same actor as this component and set the first one to be the target
-		TInlineComponentArray<UCameraComponent*> CameraComponents;
-		if (AActor* Owner = GetOwner())
-		{
-			Owner->GetComponents(CameraComponents);
-		}
-
-		if (CameraComponents.Num() > 0)
-		{
-			UCameraComponent* TargetComponent = CameraComponents[0];
-			TargetCameraComponent.ComponentProperty = TargetComponent->GetFName();
-			LastCameraComponent = TargetComponent;
-		}
+		SetCameraActor(GetOwner());
 	}
 }
 
@@ -236,6 +237,23 @@ void UCompositeViewProjectionComponent::PostEditImport()
 void UCompositeViewProjectionComponent::PostLoad()
 {
 	Super::PostLoad();
+
+PRAGMA_DISABLE_DEPRECATION_WARNINGS
+#if WITH_EDITORONLY_DATA
+	if (GetLinkerCustomVersion(FUE5MainStreamObjectVersion::GUID) < FUE5MainStreamObjectVersion::CompositeCameraReferenceRefactor)
+	{
+		const UCameraComponent* CameraComponent = Cast<UCameraComponent>(TargetCameraComponent_DEPRECATED.GetComponent(TargetCameraComponent_DEPRECATED.OtherActor.Get()));
+
+		if (IsValid(CameraComponent))
+		{
+			CameraActor = CameraComponent->GetOwner();
+		}
+
+		TargetCameraComponent_DEPRECATED.OverrideComponent.Reset();
+		TargetCameraComponent_DEPRECATED.OtherActor.Reset();
+	}
+#endif // WITH_EDITORONLY_DATA
+PRAGMA_ENABLE_DEPRECATION_WARNINGS
 }
 
 void UCompositeViewProjectionComponent::OnUnregister()
@@ -267,13 +285,14 @@ void UCompositeViewProjectionComponent::UpdateProjection(UWorld* InWorld) const
 	}
 
 	UMaterialParameterCollection* MPC = MaterialParameterCollection.Get();
+	UCameraComponent* CameraComponent = GetCameraComponent();
 
-	if (IsValid(LastCameraComponent) && IsValid(MPC))
+	if (IsValid(CameraComponent) && IsValid(MPC))
 	{
 		FMinimalViewInfo DesiredView;
 		FMatrix ViewMatrix, ProjectionMatrix, ViewProjectionMatrix;
 
-		LastCameraComponent->GetCameraView(FApp::GetDeltaTime(), DesiredView);
+		CameraComponent->GetCameraView(FApp::GetDeltaTime(), DesiredView);
 
 		UGameplayStatics::GetViewProjectionMatrix(DesiredView, ViewMatrix, ProjectionMatrix, ViewProjectionMatrix);
 
@@ -310,13 +329,6 @@ void UCompositeViewProjectionComponent::TickComponent(float DeltaTime, ELevelTic
 #if WITH_EDITOR
 void UCompositeViewProjectionComponent::PostEditChangeProperty(FPropertyChangedEvent& PropertyChangedEvent)
 {
-	const FName PropertyName = (PropertyChangedEvent.Property != NULL) ? PropertyChangedEvent.Property->GetFName() : NAME_None;
-
-	if (PropertyName == GET_MEMBER_NAME_CHECKED(ThisClass, TargetCameraComponent))
-	{
-		LastCameraComponent = Cast<UCameraComponent>(TargetCameraComponent.GetComponent(GetOwner()));
-	}
-
 	// Force MPC updates on any property change.
 	ForceUpdate();
 
diff --git a/Engine/Plugins/Compositing/Composite/Source/Composite/Private/CompositeActor.cpp b/Engine/Plugins/Compositing/Composite/Source/Composite/Private/CompositeActor.cpp
--- a/Engine/Plugins/Compositing/Composite/Source/Composite/Private/CompositeActor.cpp
+++ b/Engine/Plugins/Compositing/Composite/Source/Composite/Private/CompositeActor.cpp
@@ -27,6 +27,7 @@
 #include "SceneView.h"
 #include "Components/CompositePlateTexturePreviewComponent.h"
 #include "UObject/FortniteMainBranchObjectVersion.h"
+#include "UObject/UE5MainStreamObjectVersion.h"
 #include "UObject/UE5ReleaseStreamObjectVersion.h"
 #include "UObject/ObjectSaveContext.h"
 #include "UObject/UObjectIterator.h"
@@ -278,7 +279,9 @@ ACompositeActor::ACompositeActor(const FObjectInitializer& ObjectInitializer)
 #endif
 }
 
+PRAGMA_DISABLE_DEPRECATION_WARNINGS
 ACompositeActor::~ACompositeActor() = default;
+PRAGMA_ENABLE_DEPRECATION_WARNINGS
 
 void ACompositeActor::PostLoad()
 {
@@ -322,6 +325,21 @@ void ACompositeActor::PostLoad()
 			}
 		}
 	}
+
+	if (GetLinkerCustomVersion(FUE5MainStreamObjectVersion::GUID) < FUE5MainStreamObjectVersion::CompositeCameraReferenceRefactor)
+	{
+		PRAGMA_DISABLE_DEPRECATION_WARNINGS
+		const UCameraComponent* CameraComponent = Cast<UCameraComponent>(Camera_DEPRECATED.GetComponent(Camera_DEPRECATED.OtherActor.Get()));
+
+		if (IsValid(CameraComponent))
+		{
+			CameraActor = CameraComponent->GetOwner();
+		}
+
+		Camera_DEPRECATED.OverrideComponent.Reset();
+		Camera_DEPRECATED.OtherActor.Reset();
+		PRAGMA_ENABLE_DEPRECATION_WARNINGS
+	}
 #endif // WITH_EDITORONLY_DATA
 
 	if (VPRole.IsValid())
@@ -545,11 +563,11 @@ void ACompositeActor::PostEditChangeChainProperty(struct FPropertyChangedChainEv
 			}
 		}
 	}
-	else if (PropertyName == GET_MEMBER_NAME_CHECKED(ThisClass, Camera))
+	else if (PropertyName == GET_MEMBER_NAME_CHECKED(ThisClass, CameraActor))
 	{
 		if (IsValid(ViewProjectionComponent))
 		{
-			ViewProjectionComponent->SetTargetComponent(Camera);
+			ViewProjectionComponent->SetCameraActor(CameraActor);
 		}
 	}
 	else if (PropertyName == GET_MEMBER_NAME_CHECKED(ThisClass, VPRole))
@@ -669,18 +687,33 @@ void ACompositeActor::SetEnabled(bool bInEnabled)
 	PropagateStateChange(IsRendering() ? ECompositeStateChangeType::Activate : ECompositeStateChangeType::Deactivate);
 }
 
-FComponentReference& ACompositeActor::GetCamera()
+UCameraComponent* ACompositeActor::GetCameraComponent() const
+{
+	if (CameraActor.IsValid())
+	{
+		return CameraActor->GetComponentByClass<UCameraComponent>();
+	}
+
+	return nullptr;
+}
+
+void ACompositeActor::SetCameraComponent(UCameraComponent* InComponent)
+{
+	SetCameraActor(InComponent ? InComponent->GetOwner() : nullptr);
+}
+
+const TSoftObjectPtr<AActor>& ACompositeActor::GetCameraActor()
 {
-	return Camera;
+	return CameraActor;
 }
 
-void ACompositeActor::SetCamera(const FComponentReference& InComponentReference)
+void ACompositeActor::SetCameraActor(const TSoftObjectPtr<AActor>& InActor)
 {
-	Camera = InComponentReference;
+	CameraActor = InActor;
 	
 	if (IsValid(ViewProjectionComponent))
 	{
-		ViewProjectionComponent->SetTargetComponent(Camera);
+		ViewProjectionComponent->SetCameraActor(InActor);
 	}
 }
 
diff --git a/Engine/Plugins/Compositing/Composite/Source/Composite/Private/Layers/CompositeLayerPlanarReflection.cpp b/Engine/Plugins/Compositing/Composite/Source/Composite/Private/Layers/CompositeLayerPlanarReflection.cpp
--- a/Engine/Plugins/Compositing/Composite/Source/Composite/Private/Layers/CompositeLayerPlanarReflection.cpp
+++ b/Engine/Plugins/Compositing/Composite/Source/Composite/Private/Layers/CompositeLayerPlanarReflection.cpp
@@ -83,8 +83,7 @@ namespace UE
 					return false;
 				}
 
-				const UCameraComponent* CompositeCamera = Cast<UCameraComponent>(CompositeActor.GetCamera().GetComponent(nullptr));
-				const AActor* CompositeCameraActor = CompositeCamera ? CompositeCamera->GetOwner() : nullptr;
+				const AActor* CompositeCameraActor = CompositeActor.GetCameraActor().Get();
 
 				if (ULevelEditorSubsystem* LevelEditorSubsystem = GEditor->GetEditorSubsystem<ULevelEditorSubsystem>())
 				{
@@ -395,8 +394,7 @@ void UCompositeLayerPlanarReflection::UpdatePlanarReflectionTransform(USceneCapt
 		return;
 	}
 
-	FComponentReference& CameraRef = CompositeActor->GetCamera();
-	UCameraComponent* CameraComponent = Cast<UCameraComponent>(CameraRef.GetComponent(nullptr));
+	UCameraComponent* CameraComponent = CompositeActor->GetCameraComponent();
 	const UWorld* World = GetWorld();
 
 	if (!World)
diff --git a/Engine/Plugins/Compositing/Composite/Source/Composite/Private/Passes/CompositePassCenteredScale.cpp b/Engine/Plugins/Compositing/Composite/Source/Composite/Private/Passes/CompositePassCenteredScale.cpp
--- a/Engine/Plugins/Compositing/Composite/Source/Composite/Private/Passes/CompositePassCenteredScale.cpp
+++ b/Engine/Plugins/Compositing/Composite/Source/Composite/Private/Passes/CompositePassCenteredScale.cpp
@@ -134,7 +134,7 @@ FVector2f UCompositePassCenteredScale::CalculateScale() const
 
 	if (CompositeActor)
 	{
-		CameraComponent = Cast<UCameraComponent>(CompositeActor->GetCamera().GetComponent(nullptr));
+		CameraComponent = CompositeActor->GetCameraComponent();
 	}
 
 	FVector2f ScaleUV = FVector2f::One();
diff --git a/Engine/Plugins/Compositing/Composite/Source/Composite/Public/Components/CompositeViewProjectionComponent.h b/Engine/Plugins/Compositing/Composite/Source/Composite/Public/Components/CompositeViewProjectionComponent.h
--- a/Engine/Plugins/Compositing/Composite/Source/Composite/Public/Components/CompositeViewProjectionComponent.h
+++ b/Engine/Plugins/Compositing/Composite/Source/Composite/Public/Components/CompositeViewProjectionComponent.h
@@ -63,34 +63,46 @@ public:
 	UFUNCTION(BlueprintCallable, Category = "CameraProjection")
 	UE_API void ForceUpdate();
 
-	/** Target component getter. */
+	/** Camera component getter. */
 	UFUNCTION(BlueprintGetter)
-	UE_API FComponentReference& GetTargetComponent();
+	UE_API UCameraComponent* GetCameraComponent() const;
 
-	/** Target component setter. */
+	/** Camera component setter. */
 	UFUNCTION(BlueprintSetter)
-	UE_API void SetTargetComponent(const FComponentReference& InComponentReference);
+	UE_API void SetCameraComponent(UCameraComponent* InComponent);
+
+	/** Camera actor getter. */
+	UFUNCTION(BlueprintGetter)
+	UE_API const TSoftObjectPtr<AActor>& GetCameraActor();
+
+	/** Camera actor setter. */
+	UFUNCTION(BlueprintSetter)
+	UE_API void SetCameraActor(const TSoftObjectPtr<AActor>& InActor);
 
 private:
-	/** The CineCameraComponent whose view-projection matrix is used. */
-	UPROPERTY(EditInstanceOnly, BlueprintGetter = GetTargetComponent, BlueprintSetter = SetTargetComponent, Category = "CameraProjection", meta = (AllowPrivateAccess, UseComponentPicker, AllowedClasses = "/Script/Engine.CameraComponent, /Script/CinematicCamera.CineCameraComponent"))
-	FComponentReference TargetCameraComponent;
+	/** The camera whose view-projection matrix is used. */
+	UPROPERTY(EditAnywhere, BlueprintGetter = GetCameraActor, BlueprintSetter = SetCameraActor, Category = "CameraProjection", meta = (AllowPrivateAccess))
+	TSoftObjectPtr<AActor> CameraActor;
 
-	/** If TargetCameraComponent is not set, initialize it to the first CineCameraComponent on the same actor as this component */
+#if WITH_EDITORONLY_DATA
+	UE_DEPRECATED(5.8, "TargetCameraComponent has been deprecated, use CameraActor instead.")
+	UPROPERTY(meta = (DeprecatedProperty, DeprecationMessage = "TargetCameraComponent has been deprecated, use CameraActor instead."))
+	FComponentReference TargetCameraComponent_DEPRECATED;
+#endif
+
+	/** Initialize CameraActor from the outer CompositeActor's camera, or default to the owning actor. */
 	void InitDefaultCamera();
 	
 	/** Update the Material Parameter Collection with the camera view transform information. */
 	void UpdateProjection(UWorld* InWorld) const;
 
-	/** Cached most recent target camera, used to clean up the old camera when the user changes the target */
-	UPROPERTY()
-	TObjectPtr<UCameraComponent> LastCameraComponent = nullptr;
-
 	/** Cached view projection matrix, to avoid needless MPC updates. */
 	mutable FMatrix LastViewProjectionMatrix = FMatrix::Identity;
 
 	/** Post-viewport update callback. */
 	FDelegateHandle OnWorldPreSendAllEndOfFrameUpdatesHandle;
+
+	friend class FCompositeViewProjectionComponentCustomization;
 };
 
 #undef UE_API
diff --git a/Engine/Plugins/Compositing/Composite/Source/Composite/Public/CompositeActor.h b/Engine/Plugins/Compositing/Composite/Source/Composite/Public/CompositeActor.h
--- a/Engine/Plugins/Compositing/Composite/Source/Composite/Public/CompositeActor.h
+++ b/Engine/Plugins/Compositing/Composite/Source/Composite/Public/CompositeActor.h
@@ -14,6 +14,7 @@
 #define UE_API COMPOSITE_API
 
 class ACameraActor;
+class UCameraComponent;
 class UCompositeShadowReflectionCatcherComponent;
 class UPrimitiveComponent;
 class UCompositeSceneCapture2DComponent;
@@ -101,11 +102,19 @@ public:
 
 	/** Camera component getter. */
 	UFUNCTION(BlueprintGetter)
-	UE_API FComponentReference& GetCamera();
+	UE_API UCameraComponent* GetCameraComponent() const;
 
 	/** Camera component setter. */
 	UFUNCTION(BlueprintSetter)
-	UE_API void SetCamera(const FComponentReference& InComponentReference);
+	UE_API void SetCameraComponent(UCameraComponent* InComponent);
+
+	/** Camera actor getter. */
+	UFUNCTION(BlueprintGetter)
+	UE_API const TSoftObjectPtr<AActor>& GetCameraActor();
+
+	/** Camera actor setter. */
+	UFUNCTION(BlueprintSetter)
+	UE_API void SetCameraActor(const TSoftObjectPtr<AActor>& InActor);
 
 	/** Get the composite layers. */
 	UFUNCTION(BlueprintGetter)
@@ -162,8 +171,8 @@ private:
 public:
 private:
 	/** The primary camera used for the composite. Used by the camera view projection component for continuous view-projection matrix updates. */
-	UPROPERTY(EditInstanceOnly, BlueprintGetter = GetCamera, BlueprintSetter = SetCamera, Category = "Composite", meta = (AllowPrivateAccess, UseComponentPicker, AllowAnyActor, AllowedClasses = "/Script/Engine.CameraComponent, /Script/CinematicCamera.CineCameraComponent"))
-	FComponentReference Camera;
+	UPROPERTY(EditAnywhere, BlueprintGetter = GetCameraActor, BlueprintSetter = SetCameraActor, Category = "Composite", meta = (AllowPrivateAccess))
+	TSoftObjectPtr<AActor> CameraActor;
 
 	/** Array of composite layers for merging or processing images. */
 	UPROPERTY(EditAnywhere, BlueprintGetter = GetCompositeLayers, BlueprintSetter = SetCompositeLayers, Instanced, Category = "Composite", meta = (DisplayName = "Layers"))
@@ -222,6 +231,10 @@ private:
 	TMap<TWeakObjectPtr<const UCompositeLayerBase>, FSceneCaptureComponentArray> SceneCapturesPerLayer;
 
 #if WITH_EDITORONLY_DATA
+	UE_DEPRECATED(5.8, "Camera has been deprecated, use CameraActor instead.")
+	UPROPERTY(meta = (DeprecatedProperty, DeprecationMessage = "Camera has been deprecated, use CameraActor instead."))
+	FComponentReference Camera_DEPRECATED;
+
 	/** Used to update components after we remove layers. */
 	TArray<UCompositeLayerBase*> PreEditCompositeLayers;
 
diff --git a/Engine/Plugins/Compositing/Composite/Source/CompositeEditor/Private/CompositeEditorModule.cpp b/Engine/Plugins/Compositing/Composite/Source/CompositeEditor/Private/CompositeEditorModule.cpp
--- a/Engine/Plugins/Compositing/Composite/Source/CompositeEditor/Private/CompositeEditorModule.cpp
+++ b/Engine/Plugins/Compositing/Composite/Source/CompositeEditor/Private/CompositeEditorModule.cpp
@@ -15,6 +15,7 @@
 #include "Modules/ModuleManager.h"
 #include "Widgets/Notifications/SNotificationList.h"
 #include "Components/CompositeDepthMeshComponent.h"
+#include "Components/CompositeViewProjectionComponent.h"
 #include "CompositeActor.h"
 #include "CompositeEditorCommands.h"
 #include "CompositeEditorStyle.h"
@@ -23,6 +24,7 @@
 #include "ColorGradingDrawer/ColorGradingHierarchyConfig_Composite.h"
 #include "ColorGradingDrawer/ColorGradingSelectionHandler_Composite.h"
 #include "Customizations/CompositeActorCustomization.h"
+#include "Customizations/CompositeViewProjectionComponentCustomization.h"
 #include "Customizations/CompositeLayerBaseCustomization.h"
 #include "Customizations/CompositeLayerPlateCustomization.h"
 #include "Customizations/CompositeLayerSceneCaptureCustomization.h"
@@ -232,6 +234,11 @@ void FCompositeEditorModule::RegisterCustomizations()
 		FOnGetDetailCustomizationInstance::CreateStatic(&FCompositeActorCustomization::MakeInstance)
 	);
 
+	PropertyModule.RegisterCustomClassLayout(
+		UCompositeViewProjectionComponent::StaticClass()->GetFName(),
+		FOnGetDetailCustomizationInstance::CreateStatic(&FCompositeViewProjectionComponentCustomization::MakeInstance)
+	);
+
 	PropertyModule.RegisterCustomClassLayout(
 		UCompositeDepthMeshComponent::StaticClass()->GetFName(),
 		FOnGetDetailCustomizationInstance::CreateStatic(&FCompositeDepthMeshComponentCustomization::MakeInstance)
@@ -255,6 +262,7 @@ void FCompositeEditorModule::UnregisterCustomizations()
 			PropertyModule->UnregisterCustomClassLayout(UCompositeLayerShadowReflection::StaticClass()->GetFName());
 			PropertyModule->UnregisterCustomClassLayout(UCompositeLayerSceneCapture::StaticClass()->GetFName());
 			PropertyModule->UnregisterCustomClassLayout(UCompositeLayerSingleLightShadow::StaticClass()->GetFName());
+			PropertyModule->UnregisterCustomClassLayout(UCompositeViewProjectionComponent::StaticClass()->GetFName());
 			PropertyModule->UnregisterCustomClassLayout(UCompositeDepthMeshComponent::StaticClass()->GetFName());
 			PropertyModule->UnregisterCustomClassLayout(UCompositePassColorKeyer::StaticClass()->GetFName());
 		}
diff --git a/Engine/Plugins/Compositing/Composite/Source/CompositeEditor/Private/Customizations/CompositeActorCustomization.cpp b/Engine/Plugins/Compositing/Composite/Source/CompositeEditor/Private/Customizations/CompositeActorCustomization.cpp
--- a/Engine/Plugins/Compositing/Composite/Source/CompositeEditor/Private/Customizations/CompositeActorCustomization.cpp
+++ b/Engine/Plugins/Compositing/Composite/Source/CompositeEditor/Private/Customizations/CompositeActorCustomization.cpp
@@ -3,6 +3,7 @@
 #include "CompositeActorCustomization.h"
 
 #include "CompositeActor.h"
+#include "CompositeCustomizationHelpers.h"
 #include "CompositeEditorStyle.h"
 #include "DetailCategoryBuilder.h"
 #include "DetailLayoutBuilder.h"
@@ -71,6 +72,9 @@ void FCompositeActorCustomization::CustomizeDetails(IDetailLayoutBuilder& Detail
 			]
 		]
 	];
+
+	TSharedPtr<IPropertyHandle> CameraActorHandle = DetailLayout.GetProperty(GET_MEMBER_NAME_CHECKED(ACompositeActor, CameraActor));
+	CompositeCustomizationHelpers::CustomizeCameraActorProperty(DetailLayout, CameraActorHandle);
 }
 
 FReply FCompositeActorCustomization::OpenComposurePanel()
diff --git a/Engine/Plugins/Compositing/Composite/Source/CompositeEditor/Private/Customizations/CompositeCustomizationHelpers.h b/Engine/Plugins/Compositing/Composite/Source/CompositeEditor/Private/Customizations/CompositeCustomizationHelpers.h
new file mode 100644
--- /dev/null
+++ b/Engine/Plugins/Compositing/Composite/Source/CompositeEditor/Private/Customizations/CompositeCustomizationHelpers.h
@@ -0,0 +1,42 @@
+// Copyright Epic Games, Inc. All Rights Reserved.
+
+#pragma once
+
+#include "Camera/CameraComponent.h"
+#include "DetailLayoutBuilder.h"
+#include "PropertyCustomizationHelpers.h" // SObjectPropertyEntryBox
+
+namespace CompositeCustomizationHelpers
+{
+
+/** Customize a CameraActor soft object property to show an actor picker filtered to actors with a UCameraComponent. */
+inline void CustomizeCameraActorProperty(IDetailLayoutBuilder& DetailLayout, TSharedPtr<IPropertyHandle> CameraActorHandle)
+{
+	if (CameraActorHandle.IsValid() && CameraActorHandle->IsValidHandle())
+	{
+		DetailLayout.EditDefaultProperty(CameraActorHandle)->CustomWidget()
+			.NameContent()
+			[
+				CameraActorHandle->CreatePropertyNameWidget()
+			]
+			.ValueContent()
+			.MinDesiredWidth(250.f)
+			[
+				SNew(SObjectPropertyEntryBox)
+					.IsEnabled(true)
+					.AllowedClass(AActor::StaticClass())
+					.AllowClear(true)
+					.DisplayBrowse(true)
+					.DisplayUseSelected(true)
+					.PropertyHandle(CameraActorHandle)
+					.OnShouldFilterActor(FOnShouldFilterActor::CreateLambda([](const AActor* InActor)
+						{
+							// NOTE: For actors in SObjectPropertyEntryBox, filtered means allowed (inverse for assets).
+							return InActor && InActor->GetComponentByClass<UCameraComponent>();
+						})
+					)
+			];
+	}
+}
+
+} // namespace CompositeCustomizationHelpers
diff --git a/Engine/Plugins/Compositing/Composite/Source/CompositeEditor/Private/Customizations/CompositeLayerPlateCustomization.cpp b/Engine/Plugins/Compositing/Composite/Source/CompositeEditor/Private/Customizations/CompositeLayerPlateCustomization.cpp
--- a/Engine/Plugins/Compositing/Composite/Source/CompositeEditor/Private/Customizations/CompositeLayerPlateCustomization.cpp
+++ b/Engine/Plugins/Compositing/Composite/Source/CompositeEditor/Private/Customizations/CompositeLayerPlateCustomization.cpp
@@ -367,7 +367,7 @@ void FCompositeLayerPlateCustomization::CreateCompositeMeshActor()
 	}
 
 	FTransform SpawnTransform = FTransform::Identity;
-	if (UCameraComponent* CameraComponent = Cast<UCameraComponent>(CompositeActor->GetCamera().GetComponent(nullptr)))
+	if (UCameraComponent* CameraComponent = CompositeActor->GetCameraComponent())
 	{
 		// Size of one wall of the composite mesh
 		constexpr float MeshSize = 1000.0f;
diff --git a/Engine/Plugins/Compositing/Composite/Source/CompositeEditor/Private/Customizations/CompositeViewProjectionComponentCustomization.cpp b/Engine/Plugins/Compositing/Composite/Source/CompositeEditor/Private/Customizations/CompositeViewProjectionComponentCustomization.cpp
new file mode 100644
--- /dev/null
+++ b/Engine/Plugins/Compositing/Composite/Source/CompositeEditor/Private/Customizations/CompositeViewProjectionComponentCustomization.cpp
@@ -0,0 +1,21 @@
+// Copyright Epic Games, Inc. All Rights Reserved.
+
+#include "CompositeViewProjectionComponentCustomization.h"
+
+#include "CompositeCustomizationHelpers.h"
+#include "Components/CompositeViewProjectionComponent.h"
+
+#define LOCTEXT_NAMESPACE "FCompositeViewProjectionComponentCustomization"
+
+TSharedRef<IDetailCustomization> FCompositeViewProjectionComponentCustomization::MakeInstance()
+{
+	return MakeShared<FCompositeViewProjectionComponentCustomization>();
+}
+
+void FCompositeViewProjectionComponentCustomization::CustomizeDetails(IDetailLayoutBuilder& DetailLayout)
+{
+	TSharedPtr<IPropertyHandle> CameraActorHandle = DetailLayout.GetProperty(GET_MEMBER_NAME_CHECKED(UCompositeViewProjectionComponent, CameraActor));
+	CompositeCustomizationHelpers::CustomizeCameraActorProperty(DetailLayout, CameraActorHandle);
+}
+
+#undef LOCTEXT_NAMESPACE
diff --git a/Engine/Plugins/Compositing/Composite/Source/CompositeEditor/Private/Customizations/CompositeViewProjectionComponentCustomization.h b/Engine/Plugins/Compositing/Composite/Source/CompositeEditor/Private/Customizations/CompositeViewProjectionComponentCustomization.h
new file mode 100644
--- /dev/null
+++ b/Engine/Plugins/Compositing/Composite/Source/CompositeEditor/Private/Customizations/CompositeViewProjectionComponentCustomization.h
@@ -0,0 +1,19 @@
+// Copyright Epic Games, Inc. All Rights Reserved.
+
+#pragma once
+
+#include "IDetailCustomization.h"
+
+/**
+ * Customization for the view projection component.
+ */
+class FCompositeViewProjectionComponentCustomization : public IDetailCustomization
+{
+public:
+	/** Makes a new instance of the details customization */
+	static TSharedRef<IDetailCustomization> MakeInstance();
+
+	// IDetailCustomization interface
+	virtual void CustomizeDetails(IDetailLayoutBuilder& DetailLayout) override;
+	// End of IDetailCustomization interface
+};
diff --git a/Engine/Plugins/Compositing/Composite/Source/CompositeEditor/Private/UI/SCompositeEditorPanel.cpp b/Engine/Plugins/Compositing/Composite/Source/CompositeEditor/Private/UI/SCompositeEditorPanel.cpp
--- a/Engine/Plugins/Compositing/Composite/Source/CompositeEditor/Private/UI/SCompositeEditorPanel.cpp
+++ b/Engine/Plugins/Compositing/Composite/Source/CompositeEditor/Private/UI/SCompositeEditorPanel.cpp
@@ -4,6 +4,7 @@
 
 #include "CompositeActor.h"
 #include "CompositeEditorStyle.h"
+#include "Customizations/CompositeCustomizationHelpers.h"
 #include "DetailCategoryBuilder.h"
 #include "DetailLayoutBuilder.h"
 #include "DetailWidgetRow.h"
@@ -108,6 +109,9 @@ public:
 		
 		DetailBuilder.HideProperty(GET_MEMBER_NAME_CHECKED(ACompositeActor, bIsActive));
 		DetailBuilder.HideProperty(GET_MEMBER_NAME_CHECKED(ACompositeActor, CompositeLayers));
+
+		TSharedPtr<IPropertyHandle> CameraActorHandle = DetailBuilder.GetProperty(GET_MEMBER_NAME_CHECKED(ACompositeActor, CameraActor));
+		CompositeCustomizationHelpers::CustomizeCameraActorProperty(DetailBuilder, CameraActorHandle);
 	}
 
 private:
diff --git a/Engine/Plugins/Compositing/Composite/Source/CompositeEditor/Private/UI/SCompositePanelLayerTree.cpp b/Engine/Plugins/Compositing/Composite/Source/CompositeEditor/Private/UI/SCompositePanelLayerTree.cpp
--- a/Engine/Plugins/Compositing/Composite/Source/CompositeEditor/Private/UI/SCompositePanelLayerTree.cpp
+++ b/Engine/Plugins/Compositing/Composite/Source/CompositeEditor/Private/UI/SCompositePanelLayerTree.cpp
@@ -1566,7 +1566,7 @@ AActor* SCompositePanelLayerTree::GetActiveCompositeActorCamera() const
 
 	if (ActiveCompositeActor.IsValid())
 	{
-		return ActiveCompositeActor->GetCamera().OtherActor.Get();
+		return ActiveCompositeActor->GetCameraActor().Get();
 	}
 
 	return nullptr;
@@ -2094,7 +2094,7 @@ void SCompositePanelLayerTree::OnCompositeActorActivated(const TStrongObjectPtr<
 	// Keep last active actor up-to-date
 	LastActiveActor = CompositeActor.Get();
 
-	AActor* CompositeCameraActor = CompositeActor->GetCamera().OtherActor.Get();
+	AActor* CompositeCameraActor = CompositeActor->GetCameraActor().Get();
 
 	if (!IsValid(CompositeCameraActor))
 	{
@@ -2270,10 +2270,10 @@ void SCompositePanelLayerTree::OnObjectPropertyChanged(UObject* Object, FPropert
 			constexpr bool bPreserveSelection = true;
 			RefreshCompositeActorLayers(Cast<ACompositeActor>(Object), bRefreshTree, bPreserveSelection);
 		}
-		else if (PropertyName == GET_MEMBER_NAME_CHECKED(ACompositeActor, Camera))
+		else if (PropertyName == GET_MEMBER_NAME_CHECKED(ACompositeActor, CameraActor))
 		{
 			ACompositeActor* CompositeActor = Cast<ACompositeActor>(Object);
-			AActor* CompositeCameraActor = CompositeActor->GetCamera().OtherActor.Get();
+			AActor* CompositeCameraActor = CompositeActor->GetCameraActor().Get();
 
 			if (!IsValid(CompositeCameraActor))
 			{
diff --git a/Engine/Source/Runtime/Core/Public/UObject/UE5MainStreamObjectVersions.inl b/Engine/Source/Runtime/Core/Public/UObject/UE5MainStreamObjectVersions.inl
--- a/Engine/Source/Runtime/Core/Public/UObject/UE5MainStreamObjectVersions.inl
+++ b/Engine/Source/Runtime/Core/Public/UObject/UE5MainStreamObjectVersions.inl
@@ -374,5 +374,8 @@ UE5_MAIN_VERSION(SkinnedMeshInstanceDataSerializationV2, 120)
 // Added new material compilation validation for runtime virtual textures
 UE5_MAIN_VERSION(RuntimeVirtualTextureMaterialValidation, 121)
 
+// Refactored the composite (plugin) camera actor references.
+UE5_MAIN_VERSION(CompositeCameraReferenceRefactor, 122)
+
 // -----<new versions can be added above this line>-------------------------------------------------
 #undef UE5_MAIN_VERSION
\ No newline at end of file
