Composure: Refactor CenteredScale pass for clarity and simplification

  - Add new Unpad scale mode for per-axis black bar removal via pixel values (FIntPoint)
  - Add ECompositePassFitMode enum to control how AR mismatches are resolved
  - Rename SourceAspectRatio/TargetAspectRatio to ContainerAspectRatio/ContentAspectRatio
  - Rename ScaleFactor to ManualScale
  - Add EditConditionHides to all conditional properties so they hide when their mode is inactive
  - Add UTextureRenderTarget2D support for plate texture size fetching
  - Simplify IsActive() to check config state instead of recomputing CalculateScale()
  - Optimize GetProxy() to compute CalculateScale() once and skip on identity

#Synced-CL 51806324

diff --git a/Engine/Plugins/Compositing/Composite/Config/DefaultComposite.ini b/Engine/Plugins/Compositing/Composite/Config/DefaultComposite.ini
--- a/Engine/Plugins/Compositing/Composite/Config/DefaultComposite.ini
+++ b/Engine/Plugins/Compositing/Composite/Config/DefaultComposite.ini
@@ -41,3 +41,6 @@
 
 +PropertyRedirects=(OldName="/Script/Composite.CompositeViewProjectionComponent.MatrixParameterName", NewName="/Script/Composite.CompositeViewProjectionComponent.ViewProjectionMatrixParameter")
 +ClassRedirects=(OldName="/Script/Composite.CompositePassColorGrade",NewName="/Script/Composite.CompositePassColorGrading")
++PropertyRedirects=(OldName="/Script/Composite.CompositePassCenteredScale.SourceAspectRatio", NewName="/Script/Composite.CompositePassCenteredScale.ContainerAspectRatio")
++PropertyRedirects=(OldName="/Script/Composite.CompositePassCenteredScale.TargetAspectRatio", NewName="/Script/Composite.CompositePassCenteredScale.ContentAspectRatio")
++PropertyRedirects=(OldName="/Script/Composite.CompositePassCenteredScale.ScaleFactor", NewName="/Script/Composite.CompositePassCenteredScale.ManualScale")
diff --git a/Engine/Plugins/Compositing/Composite/Source/Composite/Private/Passes/CompositePassCenteredScale.cpp b/Engine/Plugins/Compositing/Composite/Source/Composite/Private/Passes/CompositePassCenteredScale.cpp
--- a/Engine/Plugins/Compositing/Composite/Source/Composite/Private/Passes/CompositePassCenteredScale.cpp
+++ b/Engine/Plugins/Compositing/Composite/Source/Composite/Private/Passes/CompositePassCenteredScale.cpp
@@ -4,6 +4,7 @@
 
 #include "Camera/CameraComponent.h"
 #include "CompositeActor.h"
+#include "Engine/TextureRenderTarget2D.h"
 #include "Engine/Texture2D.h"
 #include "Layers/CompositeLayerPlate.h"
 #include "MediaTexture.h"
@@ -93,9 +94,11 @@ namespace UE
 UCompositePassCenteredScale::UCompositePassCenteredScale(const FObjectInitializer& ObjectInitializer)
 	: Super(ObjectInitializer)
 	, ScaleMode(ECompositePassScaleMode::Automatic)
-	, SourceAspectRatio(16.0f, 9.0f)
-	, TargetAspectRatio(16.0f, 9.0f)
-	, ScaleFactor(1.0f, 1.0f)
+	, ContainerAspectRatio(16.0f, 9.0f)
+	, ContentAspectRatio(16.0f, 9.0f)
+	, FitMode(ECompositePassFitMode::NonUniform)
+	, UnpadPixels(FIntPoint::ZeroValue)
+	, ManualScale(1.0f, 1.0f)
 	, OverscanUncropMode(ECompositePassOverscanUncropMode::Automatic)
 	, Overscan(0.0f)
 {
@@ -107,14 +110,16 @@ bool UCompositePassCenteredScale::GetProxy(const UE::CompositeCore::FPassInputDe
 {
 	using namespace UE::CompositeCore;
 
-	// Skip pass if there is no active scaling
-	if (!IsActive())
+	// Compute scale once and skip pass if there is no active scaling.
+	const FVector2f ScaleUV = CalculateScale();
+
+	if (!IsEnabled() || ScaleUV.Equals(FVector2f::One(), UE_SMALL_NUMBER))
 	{
 		return false;
 	}
 
 	FCenteredScalePassProxy* Proxy = InFrameAllocator.Create<FCenteredScalePassProxy>(FPassInputDeclArray{ InputDecl });
-	Proxy->ScaleUV = CalculateScale();
+	Proxy->ScaleUV = ScaleUV;
 
 	OutProxy = Proxy;
 	return true;
@@ -122,13 +127,13 @@ bool UCompositePassCenteredScale::GetProxy(const UE::CompositeCore::FPassInputDe
 
 bool UCompositePassCenteredScale::IsActive() const
 {
-	return IsEnabled() && !CalculateScale().Equals(FVector2f::One(), UE_SMALL_NUMBER);
+	return IsEnabled()
+		&& (ScaleMode != ECompositePassScaleMode::None
+			|| OverscanUncropMode != ECompositePassOverscanUncropMode::None);
 }
 
 FVector2f UCompositePassCenteredScale::CalculateScale() const
 {
-	// TODO: The scale UV could be cached, though special care would be needed to acocunt for dependent property changes
-
 	ACompositeActor* CompositeActor = GetTypedOuter<ACompositeActor>();
 	UCameraComponent* CameraComponent = nullptr;
 
@@ -139,57 +144,87 @@ FVector2f UCompositePassCenteredScale::CalculateScale() const
 
 	FVector2f ScaleUV = FVector2f::One();
 
-	auto ApplyAspectRatioScaleFn = [&ScaleUV](float SourceAR, float TargetAR)
+	auto ApplyAspectRatioScaleFn = [this, &ScaleUV](float ContainerAR, float ContentAR)
 		{
-			if (SourceAR > TargetAR)
+			const bool bUniform = (FitMode == ECompositePassFitMode::Uniform);
+
+			if (ContainerAR > ContentAR)
 			{
-				ScaleUV.X = (SourceAR != 0.0f) ? TargetAR / SourceAR : 0.0f;
+				const float Factor = (ContainerAR != 0.0f) ? ContentAR / ContainerAR : 0.0f;
+
+				ScaleUV *= FVector2f(Factor, bUniform ? Factor : 1.0f);
 			}
-			else if (SourceAR < TargetAR)
+			else if (ContainerAR < ContentAR)
 			{
-				ScaleUV.Y = (TargetAR != 0.0f) ? SourceAR / TargetAR : 0.0f;
+				const float Factor = (ContentAR != 0.0f) ? ContainerAR / ContentAR : 0.0f;
+
+				ScaleUV *= FVector2f(bUniform ? Factor : 1.0f, Factor);
 			}
 		};
 
-	switch (ScaleMode)
-	{
-	case ECompositePassScaleMode::None:
-		break;
-	case ECompositePassScaleMode::Automatic:
-	{
-		float SourceAR = 16.0f / 9.0f;
+	// Fetch plate texture dimensions (used by Automatic and Unpad modes).
+	FIntVector2 TextureSize = FIntVector2::ZeroValue;
 
+	if (ScaleMode == ECompositePassScaleMode::Automatic || ScaleMode == ECompositePassScaleMode::Unpad)
+	{
 		if (UCompositeLayerPlate* PlateLayer = GetTypedOuter<UCompositeLayerPlate>())
 		{
 			if (IsValid(PlateLayer->Texture))
 			{
 				if (const UMediaTexture* MediaTexture = Cast<UMediaTexture>(PlateLayer->Texture))
 				{
-					SourceAR = MediaTexture->CurrentAspectRatio;
+					TextureSize = FIntVector2(MediaTexture->GetWidth(), MediaTexture->GetHeight());
 				}
 				else if (const UTexture2D* Texture2D = Cast<UTexture2D>(PlateLayer->Texture))
 				{
-					SourceAR = FMath::SafeDivide(Texture2D->GetSizeX(), static_cast<float>(Texture2D->GetSizeY()));
+					TextureSize = FIntVector2(Texture2D->GetSizeX(), Texture2D->GetSizeY());
+				}
+				else if (const UTextureRenderTarget2D* RenderTarget2D = Cast<UTextureRenderTarget2D>(PlateLayer->Texture))
+				{
+					TextureSize = FIntVector2(RenderTarget2D->SizeX, RenderTarget2D->SizeY);
 				}
 			}
 		}
+	}
 
-		const float TargetAR = CameraComponent ? CameraComponent->AspectRatio : SourceAR;
+	switch (ScaleMode)
+	{
+	case ECompositePassScaleMode::None:
+		break;
+	case ECompositePassScaleMode::Automatic:
+	{
+		if (TextureSize.X > 0 && TextureSize.Y > 0)
+		{
+			const float ContainerAR = FMath::SafeDivide(TextureSize.X, static_cast<float>(TextureSize.Y));
+			const float ContentAR = CameraComponent ? CameraComponent->AspectRatio : ContainerAR;
 
-		ApplyAspectRatioScaleFn(SourceAR, TargetAR);
+			ApplyAspectRatioScaleFn(ContainerAR, ContentAR);
+		}
 	}
 	break;
 	case ECompositePassScaleMode::AspectRatio:
 	{
-		const float SourceAR = FMath::SafeDivide(SourceAspectRatio.X, SourceAspectRatio.Y);
-		const float TargetAR = FMath::SafeDivide(TargetAspectRatio.X, TargetAspectRatio.Y);
+		const float ContainerAR = FMath::SafeDivide(ContainerAspectRatio.X, ContainerAspectRatio.Y);
+		const float ContentAR = FMath::SafeDivide(ContentAspectRatio.X, ContentAspectRatio.Y);
 
-		ApplyAspectRatioScaleFn(SourceAR, TargetAR);
+		ApplyAspectRatioScaleFn(ContainerAR, ContentAR);
+	}
+	break;
+	case ECompositePassScaleMode::Unpad:
+	{
+		if (TextureSize.X > 2 * UnpadPixels.X && TextureSize.X > 0)
+		{
+			ScaleUV.X *= static_cast<float>(TextureSize.X - 2 * UnpadPixels.X) / TextureSize.X;
+		}
+		if (TextureSize.Y > 2 * UnpadPixels.Y && TextureSize.Y > 0)
+		{
+			ScaleUV.Y *= static_cast<float>(TextureSize.Y - 2 * UnpadPixels.Y) / TextureSize.Y;
+		}
 	}
 	break;
 	case ECompositePassScaleMode::Manual:
-		ScaleUV.X = FMath::SafeDivide(ScaleUV.X, ScaleFactor.X);
-		ScaleUV.Y = FMath::SafeDivide(ScaleUV.Y, ScaleFactor.Y);
+		ScaleUV.X = FMath::SafeDivide(ScaleUV.X, ManualScale.X);
+		ScaleUV.Y = FMath::SafeDivide(ScaleUV.Y, ManualScale.Y);
 		break;
 
 	default:
@@ -218,4 +253,3 @@ FVector2f UCompositePassCenteredScale::CalculateScale() const
 
 	return ScaleUV;
 }
-
diff --git a/Engine/Plugins/Compositing/Composite/Source/Composite/Public/Passes/CompositePassCenteredScale.h b/Engine/Plugins/Compositing/Composite/Source/Composite/Public/Passes/CompositePassCenteredScale.h
--- a/Engine/Plugins/Compositing/Composite/Source/Composite/Public/Passes/CompositePassCenteredScale.h
+++ b/Engine/Plugins/Compositing/Composite/Source/Composite/Public/Passes/CompositePassCenteredScale.h
@@ -8,7 +8,7 @@
 
 #define UE_API COMPOSITE_API
 
-/** Primary centered scale calculation mode.*/
+/** Primary scale calculation mode. */
 UENUM(BlueprintType)
 enum class ECompositePassScaleMode : uint8
 {
@@ -18,13 +18,27 @@ enum class ECompositePassScaleMode : uint8
 	/** Automatically derive aspect ratio from the parent layer media texture & composite actor camera. */
 	Automatic,
 
-	/** Calculate the scaling factor from source & target aspect ratios.*/
+	/** Scale based on the difference between container & content aspect ratios. */
 	AspectRatio,
 
-	/** Manually define the scaling factor.*/
+	/** Remove black bar padding by specifying pixels to crop on each axis. */
+	Unpad,
+
+	/** Manually define the scale factor.*/
 	Manual,
 };
 
+/** How aspect ratio mismatch is resolved when scaling. Only applies to Automatic and AspectRatio modes. */
+UENUM(BlueprintType)
+enum class ECompositePassFitMode : uint8
+{
+	/** Scale only the mismatched axis. i.e. 9:16 content in a 9:16 camera crops only the pillarbox bars. */
+	NonUniform = 0,
+
+	/** Scale both axes equally. Preserves spatial proportions. E.g. 9:16 content in a 16:9 camera zooms uniformly to fill the frame. */
+	Uniform = 1,
+};
+
 /** Overscan uncrop derivation mode. */
 UENUM(BlueprintType)
 enum class ECompositePassOverscanUncropMode : uint8
@@ -48,7 +62,7 @@ class UCompositePassCenteredScale : public UCompositePassBase
 public:
 	/** Constructor */
 	UE_API UCompositePassCenteredScale(const FObjectInitializer& ObjectInitializer);
-	
+
 	/** Destructor */
 	UE_API ~UCompositePassCenteredScale();
 
@@ -61,30 +75,37 @@ public:
 	UE_API FVector2f CalculateScale() const;
 
 public:
-	/** Centered scale calculation mode used to rescale a media texture with black bars into an already constrained aspect ratio viewport. */
+	/** How to scale-fit the video content into the camera viewport. */
 	UPROPERTY(EditAnywhere, BlueprintReadWrite, Interp, Category = "Composite")
 	ECompositePassScaleMode ScaleMode;
 
-	/** Source container aspect ratio (or resolution). */
-	UPROPERTY(EditAnywhere, BlueprintReadWrite, Interp, Category = "Composite", meta = (EditCondition = "ScaleMode == ECompositePassScaleMode::AspectRatio"))
-	FVector2f SourceAspectRatio;
+	/** Container (source) aspect ratio, as width:height (e.g. 16:9). Automatic derives this from the plate texture. */
+	UPROPERTY(EditAnywhere, BlueprintReadWrite, Interp, Category = "Composite", meta = (EditCondition = "ScaleMode == ECompositePassScaleMode::AspectRatio", EditConditionHides))
+	FVector2f ContainerAspectRatio;
+
+	/** Content (target) aspect ratio, typically matching the cine camera, as width:height (e.g. 16:9). Automatic derives this from the camera component. */
+	UPROPERTY(EditAnywhere, BlueprintReadWrite, Interp, Category = "Composite", meta = (EditCondition = "ScaleMode == ECompositePassScaleMode::AspectRatio", EditConditionHides))
+	FVector2f ContentAspectRatio;
 
-	/** Embedded target aspect ratio (or resolution), without black bars. */
-	UPROPERTY(EditAnywhere, BlueprintReadWrite, Interp, Category = "Composite", meta = (EditCondition = "ScaleMode == ECompositePassScaleMode::AspectRatio"))
-	FVector2f TargetAspectRatio;
+	/** How to resolve aspect ratio mismatch. NonUniform scales only the mismatched axis. Uniform scales both axes equally. */
+	UPROPERTY(EditAnywhere, BlueprintReadWrite, Interp, Category = "Composite", meta = (EditCondition = "ScaleMode == ECompositePassScaleMode::Automatic || ScaleMode == ECompositePassScaleMode::AspectRatio", EditConditionHides))
+	ECompositePassFitMode FitMode;
 
-	/** Manual scale factor. */
-	UPROPERTY(EditAnywhere, BlueprintReadWrite, Interp, Category = "Composite", meta = (EditCondition = "ScaleMode == ECompositePassScaleMode::Manual"))
-	FVector2f ScaleFactor;
+	/** Pixels of black bar padding to crop from each side, per axis. For example, (656, 0) removes 656px of pillarbox from left and right. */
+	UPROPERTY(EditAnywhere, BlueprintReadWrite, Interp, Category = "Composite", meta = (EditCondition = "ScaleMode == ECompositePassScaleMode::Unpad", EditConditionHides, ClampMin = "0", UIMax = "500"))
+	FIntPoint UnpadPixels;
+
+	/** Manual scale factor per axis. Values > 1 zoom in, values < 1 zoom out. */
+	UPROPERTY(EditAnywhere, BlueprintReadWrite, Interp, Category = "Composite", meta = (EditCondition = "ScaleMode == ECompositePassScaleMode::Manual", EditConditionHides))
+	FVector2f ManualScale;
 
 	/** Uncrop calculation mode used to uncrop a viewport cropped in by overscan. */
 	UPROPERTY(EditAnywhere, BlueprintReadWrite, Category = "Composite")
 	ECompositePassOverscanUncropMode OverscanUncropMode;
 
 	/** Manual overscan to uncrop, with values from 0.0 to 1.0 matching the source camera overscan. */
-	UPROPERTY(EditAnywhere, BlueprintReadWrite, Interp, Category = "Composite", meta = (EditCondition = "OverscanUncropMode == ECompositePassOverscanUncropMode::Manual", ClampMin = "0.0", ClampMax = "1.0", UIMin = "0.0", UIMax = "1.0"))
+	UPROPERTY(EditAnywhere, BlueprintReadWrite, Interp, Category = "Composite", meta = (EditCondition = "OverscanUncropMode == ECompositePassOverscanUncropMode::Manual", EditConditionHides, ClampMin = "0.0", ClampMax = "1.0", UIMin = "0.0", UIMax = "1.0"))
 	float Overscan;
 };
 
 #undef UE_API
-
