Composure: Removed unnecessary centered scale fit mode, replaced with bScaleSingleAxis.

#Synced-CL 51870524

diff --git a/Engine/Plugins/Compositing/Composite/Source/Composite/Private/Passes/CompositePassCenteredScale.cpp b/Engine/Plugins/Compositing/Composite/Source/Composite/Private/Passes/CompositePassCenteredScale.cpp
--- a/Engine/Plugins/Compositing/Composite/Source/Composite/Private/Passes/CompositePassCenteredScale.cpp
+++ b/Engine/Plugins/Compositing/Composite/Source/Composite/Private/Passes/CompositePassCenteredScale.cpp
@@ -96,7 +96,7 @@ UCompositePassCenteredScale::UCompositePassCenteredScale(const FObjectInitialize
 	, ScaleMode(ECompositePassScaleMode::Automatic)
 	, ContainerAspectRatio(16.0f, 9.0f)
 	, ContentAspectRatio(16.0f, 9.0f)
-	, FitMode(ECompositePassFitMode::NonUniform)
+	, bScaleSingleAxis(true)
 	, UnpadPixels(FIntPoint::ZeroValue)
 	, ManualScale(1.0f, 1.0f)
 	, OverscanUncropMode(ECompositePassOverscanUncropMode::Automatic)
@@ -146,19 +146,17 @@ FVector2f UCompositePassCenteredScale::CalculateScale() const
 
 	auto ApplyAspectRatioScaleFn = [this, &ScaleUV](float ContainerAR, float ContentAR)
 		{
-			const bool bUniform = (FitMode == ECompositePassFitMode::Uniform);
-
 			if (ContainerAR > ContentAR)
 			{
 				const float Factor = (ContainerAR != 0.0f) ? ContentAR / ContainerAR : 0.0f;
 
-				ScaleUV *= FVector2f(Factor, bUniform ? Factor : 1.0f);
+				ScaleUV *= FVector2f(Factor, bScaleSingleAxis ? 1.0f : Factor);
 			}
 			else if (ContainerAR < ContentAR)
 			{
 				const float Factor = (ContentAR != 0.0f) ? ContainerAR / ContentAR : 0.0f;
 
-				ScaleUV *= FVector2f(bUniform ? Factor : 1.0f, Factor);
+				ScaleUV *= FVector2f(bScaleSingleAxis ? 1.0f : Factor, Factor);
 			}
 		};
 
diff --git a/Engine/Plugins/Compositing/Composite/Source/Composite/Public/Passes/CompositePassCenteredScale.h b/Engine/Plugins/Compositing/Composite/Source/Composite/Public/Passes/CompositePassCenteredScale.h
--- a/Engine/Plugins/Compositing/Composite/Source/Composite/Public/Passes/CompositePassCenteredScale.h
+++ b/Engine/Plugins/Compositing/Composite/Source/Composite/Public/Passes/CompositePassCenteredScale.h
@@ -28,17 +28,6 @@ enum class ECompositePassScaleMode : uint8
 	Manual,
 };
 
-/** How aspect ratio mismatch is resolved when scaling. Only applies to Automatic and AspectRatio modes. */
-UENUM(BlueprintType)
-enum class ECompositePassFitMode : uint8
-{
-	/** Scale only the mismatched axis. i.e. 9:16 content in a 9:16 camera crops only the pillarbox bars. */
-	NonUniform = 0,
-
-	/** Scale both axes equally. Preserves spatial proportions. E.g. 9:16 content in a 16:9 camera zooms uniformly to fill the frame. */
-	Uniform = 1,
-};
-
 /** Overscan uncrop derivation mode. */
 UENUM(BlueprintType)
 enum class ECompositePassOverscanUncropMode : uint8
@@ -87,9 +76,9 @@ public:
 	UPROPERTY(EditAnywhere, BlueprintReadWrite, Interp, Category = "Composite", meta = (EditCondition = "ScaleMode == ECompositePassScaleMode::AspectRatio", EditConditionHides))
 	FVector2f ContentAspectRatio;
 
-	/** How to resolve aspect ratio mismatch. NonUniform scales only the mismatched axis. Uniform scales both axes equally. */
+	/** How to resolve aspect ratio mismatch. Scales only the mismatched axis, or both axes equally. */
 	UPROPERTY(EditAnywhere, BlueprintReadWrite, Interp, Category = "Composite", meta = (EditCondition = "ScaleMode == ECompositePassScaleMode::Automatic || ScaleMode == ECompositePassScaleMode::AspectRatio", EditConditionHides))
-	ECompositePassFitMode FitMode;
+	bool bScaleSingleAxis;
 
 	/** Pixels of black bar padding to crop from each side, per axis. For example, (656, 0) removes 656px of pillarbox from left and right. */
 	UPROPERTY(EditAnywhere, BlueprintReadWrite, Interp, Category = "Composite", meta = (EditCondition = "ScaleMode == ECompositePassScaleMode::Unpad", EditConditionHides, ClampMin = "0", UIMax = "500"))
