Composure: Fixed centered scale pass black frames due to active/getproxy mistmatch in CL 51854563.

#Synced-CL 51949275

diff --git a/Engine/Plugins/Compositing/Composite/Source/Composite/Private/Passes/CompositePassCenteredScale.cpp b/Engine/Plugins/Compositing/Composite/Source/Composite/Private/Passes/CompositePassCenteredScale.cpp
--- a/Engine/Plugins/Compositing/Composite/Source/Composite/Private/Passes/CompositePassCenteredScale.cpp
+++ b/Engine/Plugins/Compositing/Composite/Source/Composite/Private/Passes/CompositePassCenteredScale.cpp
@@ -110,13 +110,9 @@ bool UCompositePassCenteredScale::GetProxy(const UE::CompositeCore::FPassInputDe
 {
 	using namespace UE::CompositeCore;
 
-	// Compute scale once and skip pass if there is no active scaling.
-	const FVector2f ScaleUV = CalculateScale();
-
-	if (!IsEnabled() || ScaleUV.Equals(FVector2f::One(), UE_SMALL_NUMBER))
-	{
-		return false;
-	}
+	// Use cached value from IsActive() if available, otherwise compute fresh.
+	const FVector2f ScaleUV = CachedScaleUV.IsSet() ? CachedScaleUV.GetValue() : CalculateScale();
+	CachedScaleUV.Reset();
 
 	FCenteredScalePassProxy* Proxy = InFrameAllocator.Create<FCenteredScalePassProxy>(FPassInputDeclArray{ InputDecl });
 	Proxy->ScaleUV = ScaleUV;
@@ -127,9 +123,11 @@ bool UCompositePassCenteredScale::GetProxy(const UE::CompositeCore::FPassInputDe
 
 bool UCompositePassCenteredScale::IsActive() const
 {
+	CachedScaleUV = CalculateScale();
+
 	return IsEnabled()
-		&& (ScaleMode != ECompositePassScaleMode::None
-			|| OverscanUncropMode != ECompositePassOverscanUncropMode::None);
+		&& (ScaleMode != ECompositePassScaleMode::None || OverscanUncropMode != ECompositePassOverscanUncropMode::None)
+		&& !CachedScaleUV.GetValue().Equals(FVector2f::One(), UE_SMALL_NUMBER);
 }
 
 FVector2f UCompositePassCenteredScale::CalculateScale() const
diff --git a/Engine/Plugins/Compositing/Composite/Source/Composite/Public/Passes/CompositePassCenteredScale.h b/Engine/Plugins/Compositing/Composite/Source/Composite/Public/Passes/CompositePassCenteredScale.h
--- a/Engine/Plugins/Compositing/Composite/Source/Composite/Public/Passes/CompositePassCenteredScale.h
+++ b/Engine/Plugins/Compositing/Composite/Source/Composite/Public/Passes/CompositePassCenteredScale.h
@@ -95,6 +95,10 @@ public:
 	/** Manual overscan to uncrop, with values from 0.0 to 1.0 matching the source camera overscan. */
 	UPROPERTY(EditAnywhere, BlueprintReadWrite, Interp, Category = "Composite", meta = (EditCondition = "OverscanUncropMode == ECompositePassOverscanUncropMode::Manual", EditConditionHides, ClampMin = "0.0", ClampMax = "1.0", UIMin = "0.0", UIMax = "1.0"))
 	float Overscan;
+
+private:
+	/** Cached scale UV computed in IsActive() and reused in GetProxy(). Reset after consumption. */
+	mutable TOptional<FVector2f> CachedScaleUV;
 };
 
 #undef UE_API
