Composure: Pass and layer naming cleanup.

* We also make type and operation composure panel columns subdued gray since they are not editable. This matches the regular details panel.
* Name number increments now works equally between passes & layers.

#Synced-CL 52008728

diff --git a/Engine/Plugins/Compositing/Composite/Source/Composite/Private/Passes/CompositePassBase.cpp b/Engine/Plugins/Compositing/Composite/Source/Composite/Private/Passes/CompositePassBase.cpp
--- a/Engine/Plugins/Compositing/Composite/Source/Composite/Private/Passes/CompositePassBase.cpp
+++ b/Engine/Plugins/Compositing/Composite/Source/Composite/Private/Passes/CompositePassBase.cpp
@@ -37,13 +37,7 @@ UCompositePassBase::UCompositePassBase(const FObjectInitializer& ObjectInitializ
 	}
 
 #if WITH_EDITOR
-	DisplayName = GetClass()->GetName();
-	
-	const int32 NameNumber = GetFName().GetNumber();
-	if (NameNumber > 1)
-	{
-		DisplayName.AppendInt(NameNumber);
-	}
+	DisplayName = GetClass()->GetDisplayNameText().ToString();
 #endif
 }
 
diff --git a/Engine/Plugins/Compositing/Composite/Source/Composite/Public/Passes/CompositePassAlphaTransform.h b/Engine/Plugins/Compositing/Composite/Source/Composite/Public/Passes/CompositePassAlphaTransform.h
--- a/Engine/Plugins/Compositing/Composite/Source/Composite/Public/Passes/CompositePassAlphaTransform.h
+++ b/Engine/Plugins/Compositing/Composite/Source/Composite/Public/Passes/CompositePassAlphaTransform.h
@@ -27,7 +27,7 @@ enum class ECompositeAlphaOverride : uint8
 };
 
 /** Alpha transform pass. */
-UCLASS(MinimalAPI, BlueprintType, Blueprintable, EditInlineNew, CollapseCategories, meta=(DisplayName="Alpha Transform Pass"))
+UCLASS(MinimalAPI, BlueprintType, Blueprintable, EditInlineNew, CollapseCategories, meta=(DisplayName="Alpha Transform"))
 class UCompositePassAlphaTransform : public UCompositePassBase
 {
 	GENERATED_BODY()
diff --git a/Engine/Plugins/Compositing/Composite/Source/Composite/Public/Passes/CompositePassBase.h b/Engine/Plugins/Compositing/Composite/Source/Composite/Public/Passes/CompositePassBase.h
--- a/Engine/Plugins/Compositing/Composite/Source/Composite/Public/Passes/CompositePassBase.h
+++ b/Engine/Plugins/Compositing/Composite/Source/Composite/Public/Passes/CompositePassBase.h
@@ -82,6 +82,29 @@ public:
 
 	/** Sets the display name of the pass */
 	void SetDisplayName(const FString& InDisplayName) { DisplayName = InDisplayName; }
+
+	/** Generates a unique display name by appending a number if siblings share the same base type name. First instance stays clean; duplicates start at 2. */
+	template<typename T>
+	static FString MakeUniqueDisplayName(const FString& BaseName, TArrayView<TObjectPtr<T>> Siblings, const UCompositePassBase* Exclude)
+	{
+		static_assert(std::is_base_of_v<UCompositePassBase, T>, "T must derive from UCompositePassBase");
+		
+		int32 MatchCount = 0;
+		for (const TObjectPtr<T>& Sibling : Siblings)
+		{
+			if (Sibling && Sibling != Exclude && Sibling->GetClass()->GetDisplayNameText().ToString() == BaseName)
+			{
+				++MatchCount;
+			}
+		}
+		
+		if (MatchCount == 0)
+		{
+			return BaseName;
+		}
+
+		return FString::Printf(TEXT("%s %d"), *BaseName, MatchCount + 1);
+	}
 #endif
 	
 protected:
diff --git a/Engine/Plugins/Compositing/Composite/Source/Composite/Public/Passes/CompositePassBlur.h b/Engine/Plugins/Compositing/Composite/Source/Composite/Public/Passes/CompositePassBlur.h
--- a/Engine/Plugins/Compositing/Composite/Source/Composite/Public/Passes/CompositePassBlur.h
+++ b/Engine/Plugins/Compositing/Composite/Source/Composite/Public/Passes/CompositePassBlur.h
@@ -16,7 +16,7 @@ enum class ECompositePassBlurMethod : uint8
 };
 
 /** Blur composite pass. */
-UCLASS(MinimalAPI, BlueprintType, Blueprintable, EditInlineNew, CollapseCategories, meta = (DisplayName = "Blur Pass"))
+UCLASS(MinimalAPI, BlueprintType, Blueprintable, EditInlineNew, CollapseCategories, meta = (DisplayName = "Blur"))
 class UCompositePassBlur : public UCompositePassBase
 {
 	GENERATED_BODY()
diff --git a/Engine/Plugins/Compositing/Composite/Source/Composite/Public/Passes/CompositePassCenteredScale.h b/Engine/Plugins/Compositing/Composite/Source/Composite/Public/Passes/CompositePassCenteredScale.h
--- a/Engine/Plugins/Compositing/Composite/Source/Composite/Public/Passes/CompositePassCenteredScale.h
+++ b/Engine/Plugins/Compositing/Composite/Source/Composite/Public/Passes/CompositePassCenteredScale.h
@@ -43,7 +43,7 @@ enum class ECompositePassOverscanUncropMode : uint8
 };
 
 /** Convenience pass to re-scale footage (with black bars) inside its container texture. */
-UCLASS(MinimalAPI, BlueprintType, Blueprintable, EditInlineNew, CollapseCategories, meta=(DisplayName="Centered Scale Pass"))
+UCLASS(MinimalAPI, BlueprintType, Blueprintable, EditInlineNew, CollapseCategories, meta=(DisplayName="Centered Scale"))
 class UCompositePassCenteredScale : public UCompositePassBase
 {
 	GENERATED_BODY()
diff --git a/Engine/Plugins/Compositing/Composite/Source/Composite/Public/Passes/CompositePassColorGrading.h b/Engine/Plugins/Compositing/Composite/Source/Composite/Public/Passes/CompositePassColorGrading.h
--- a/Engine/Plugins/Compositing/Composite/Source/Composite/Public/Passes/CompositePassColorGrading.h
+++ b/Engine/Plugins/Compositing/Composite/Source/Composite/Public/Passes/CompositePassColorGrading.h
@@ -38,7 +38,7 @@ struct FCompositeTemperatureSettings
  * Color grade pass.
  * Assumes input is in linear working color space.
 */
-UCLASS(MinimalAPI, BlueprintType, Blueprintable, EditInlineNew, CollapseCategories, meta = (DisplayName = "Color Grading Pass"))
+UCLASS(MinimalAPI, BlueprintType, Blueprintable, EditInlineNew, CollapseCategories, meta = (DisplayName = "Color Grading"))
 class UCompositePassColorGrading : public UCompositePassBase
 {
 	GENERATED_BODY()
diff --git a/Engine/Plugins/Compositing/Composite/Source/Composite/Public/Passes/CompositePassColorKeyer.h b/Engine/Plugins/Compositing/Composite/Source/Composite/Public/Passes/CompositePassColorKeyer.h
--- a/Engine/Plugins/Compositing/Composite/Source/Composite/Public/Passes/CompositePassColorKeyer.h
+++ b/Engine/Plugins/Compositing/Composite/Source/Composite/Public/Passes/CompositePassColorKeyer.h
@@ -44,7 +44,7 @@ enum class ECompositeKeyerSource : uint8
 /**
  * Color keyer with clean plate & spill removal support.
 */
-UCLASS(MinimalAPI, BlueprintType, Blueprintable, EditInlineNew, CollapseCategories, meta=(DisplayName="Color Keyer Pass"))
+UCLASS(MinimalAPI, BlueprintType, Blueprintable, EditInlineNew, CollapseCategories, meta=(DisplayName="Color Keyer"))
 class UCompositePassColorKeyer : public UCompositePassBase
 {
 	GENERATED_BODY()
diff --git a/Engine/Plugins/Compositing/Composite/Source/Composite/Public/Passes/CompositePassDistortion.h b/Engine/Plugins/Compositing/Composite/Source/Composite/Public/Passes/CompositePassDistortion.h
--- a/Engine/Plugins/Compositing/Composite/Source/Composite/Public/Passes/CompositePassDistortion.h
+++ b/Engine/Plugins/Compositing/Composite/Source/Composite/Public/Passes/CompositePassDistortion.h
@@ -20,7 +20,7 @@ enum class ECompositeDistortion : uint8
 };
 
 /** Pass to apply an Distortion transform. */
-UCLASS(MinimalAPI, BlueprintType, Blueprintable, EditInlineNew, CollapseCategories, meta = (DisplayName = "Distortion Pass"))
+UCLASS(MinimalAPI, BlueprintType, Blueprintable, EditInlineNew, CollapseCategories, meta = (DisplayName = "Distortion"))
 class UCompositePassDistortion : public UCompositePassBase
 {
 	GENERATED_BODY()
diff --git a/Engine/Plugins/Compositing/Composite/Source/Composite/Public/Passes/CompositePassFXAA.h b/Engine/Plugins/Compositing/Composite/Source/Composite/Public/Passes/CompositePassFXAA.h
--- a/Engine/Plugins/Compositing/Composite/Source/Composite/Public/Passes/CompositePassFXAA.h
+++ b/Engine/Plugins/Compositing/Composite/Source/Composite/Public/Passes/CompositePassFXAA.h
@@ -9,7 +9,7 @@
 #define UE_API COMPOSITE_API
 
 /** FXAA holdout composite pass, mostly useful for applying on the composite render pass for CG/Motion Graphics. */
-UCLASS(MinimalAPI, BlueprintType, Blueprintable, EditInlineNew, CollapseCategories, meta=(DisplayName="FXAA Pass"))
+UCLASS(MinimalAPI, BlueprintType, Blueprintable, EditInlineNew, CollapseCategories, meta=(DisplayName="FXAA"))
 class UCompositePassFXAA : public UCompositePassBase
 {
 	GENERATED_BODY()
diff --git a/Engine/Plugins/Compositing/Composite/Source/Composite/Public/Passes/CompositePassMasking.h b/Engine/Plugins/Compositing/Composite/Source/Composite/Public/Passes/CompositePassMasking.h
--- a/Engine/Plugins/Compositing/Composite/Source/Composite/Public/Passes/CompositePassMasking.h
+++ b/Engine/Plugins/Compositing/Composite/Source/Composite/Public/Passes/CompositePassMasking.h
@@ -19,7 +19,7 @@ enum class ECompositeColorChannel : uint8
 };
 
 /** Alpha masking pass. */
-UCLASS(MinimalAPI, BlueprintType, Blueprintable, EditInlineNew, CollapseCategories, meta=(DisplayName="Masking Pass"))
+UCLASS(MinimalAPI, BlueprintType, Blueprintable, EditInlineNew, CollapseCategories, meta=(DisplayName="Masking"))
 class UCompositePassMasking : public UCompositePassBase
 {
 	GENERATED_BODY()
diff --git a/Engine/Plugins/Compositing/Composite/Source/Composite/Public/Passes/CompositePassMaterial.h b/Engine/Plugins/Compositing/Composite/Source/Composite/Public/Passes/CompositePassMaterial.h
--- a/Engine/Plugins/Compositing/Composite/Source/Composite/Public/Passes/CompositePassMaterial.h
+++ b/Engine/Plugins/Compositing/Composite/Source/Composite/Public/Passes/CompositePassMaterial.h
@@ -11,7 +11,7 @@
 class UMaterialInterface;
 
 /** Post-process material holdout composite pass. */
-UCLASS(MinimalAPI, BlueprintType, Blueprintable, EditInlineNew, CollapseCategories, meta = (DisplayName = "Post-Process Material Pass"))
+UCLASS(MinimalAPI, BlueprintType, Blueprintable, EditInlineNew, CollapseCategories, meta = (DisplayName = "Post-Process Material"))
 class UCompositePassMaterial : public UCompositePassBase
 {
 	GENERATED_BODY()
diff --git a/Engine/Plugins/Compositing/Composite/Source/Composite/Public/Passes/CompositePassOpenColorIO.h b/Engine/Plugins/Compositing/Composite/Source/Composite/Public/Passes/CompositePassOpenColorIO.h
--- a/Engine/Plugins/Compositing/Composite/Source/Composite/Public/Passes/CompositePassOpenColorIO.h
+++ b/Engine/Plugins/Compositing/Composite/Source/Composite/Public/Passes/CompositePassOpenColorIO.h
@@ -13,7 +13,7 @@
 #define UE_API COMPOSITE_API
 
 /** Pass to apply an OpenColorIO transform. */
-UCLASS(MinimalAPI, BlueprintType, Blueprintable, EditInlineNew, CollapseCategories, meta = (DisplayName = "OpenColorIO Pass"))
+UCLASS(MinimalAPI, BlueprintType, Blueprintable, EditInlineNew, CollapseCategories, meta = (DisplayName = "OpenColorIO"))
 class UCompositePassOpenColorIO : public UCompositePassBase
 {
 	GENERATED_BODY()
diff --git a/Engine/Plugins/Compositing/Composite/Source/Composite/Public/Passes/CompositePassSMAA.h b/Engine/Plugins/Compositing/Composite/Source/Composite/Public/Passes/CompositePassSMAA.h
--- a/Engine/Plugins/Compositing/Composite/Source/Composite/Public/Passes/CompositePassSMAA.h
+++ b/Engine/Plugins/Compositing/Composite/Source/Composite/Public/Passes/CompositePassSMAA.h
@@ -9,7 +9,7 @@
 #define UE_API COMPOSITE_API
 
 /** SMAA composite pass, mostly useful for applying on the composite render pass for CG/Motion Graphics. */
-UCLASS(MinimalAPI, BlueprintType, Blueprintable, EditInlineNew, CollapseCategories, meta=(DisplayName="SMAA Pass"))
+UCLASS(MinimalAPI, BlueprintType, Blueprintable, EditInlineNew, CollapseCategories, meta=(DisplayName="SMAA"))
 class UCompositePassSMAA : public UCompositePassBase
 {
 	GENERATED_BODY()
diff --git a/Engine/Plugins/Compositing/Composite/Source/CompositeEditor/Private/ActorFactories/CompositeActorFactory.cpp b/Engine/Plugins/Compositing/Composite/Source/CompositeEditor/Private/ActorFactories/CompositeActorFactory.cpp
--- a/Engine/Plugins/Compositing/Composite/Source/CompositeEditor/Private/ActorFactories/CompositeActorFactory.cpp
+++ b/Engine/Plugins/Compositing/Composite/Source/CompositeEditor/Private/ActorFactories/CompositeActorFactory.cpp
@@ -35,12 +35,6 @@ void UCompositeActorFactory::PostSpawnActor(UObject* Asset, AActor* NewActor)
 		NewLayers.Add(NewObject<UCompositeLayerShadowReflection>(CompositeActor, TEXT("ShadowReflectionLayer"), RF_Transactional));
 		NewLayers.Add(NewObject<UCompositeLayerPlate>(CompositeActor, TEXT("PlateLayer"), RF_Transactional));
 
-		// Assign default display name value to all layers
-		for (TObjectPtr<UCompositeLayerBase>& Layer : NewLayers)
-		{
-			Layer->SetDisplayName(Layer->GetName());
-		}
-
 		// By default, we disable the shadow reflection catcher due to its significant added cost.
 		NewLayers[1]->SetEnabled(false);
 
diff --git a/Engine/Plugins/Compositing/Composite/Source/CompositeEditor/Private/UI/SCompositePanelLayerTree.cpp b/Engine/Plugins/Compositing/Composite/Source/CompositeEditor/Private/UI/SCompositePanelLayerTree.cpp
--- a/Engine/Plugins/Compositing/Composite/Source/CompositeEditor/Private/UI/SCompositePanelLayerTree.cpp
+++ b/Engine/Plugins/Compositing/Composite/Source/CompositeEditor/Private/UI/SCompositePanelLayerTree.cpp
@@ -465,7 +465,9 @@ private:
 				.BorderImage(FAppStyle::GetBrush("NoBorder"))
 				.Padding(ColumnHorizontalPadding, 0.0f)
 				[
-					SNew(STextBlock).Text(this, &SCompositePanelLayerTreeItemRow::GetItemType)
+					SNew(STextBlock)
+					.Text(this, &SCompositePanelLayerTreeItemRow::GetItemType)
+					.ColorAndOpacity(FSlateColor::UseSubduedForeground())
 				];
 		}
 		else if (InColumnName == CompositePanelLayerTree::LayerTreeColumn_Operation)
@@ -474,7 +476,9 @@ private:
 				.BorderImage(FAppStyle::GetBrush("NoBorder"))
 				.Padding(ColumnHorizontalPadding, 0.0f)
 				[
-					SNew(STextBlock).Text(this, &SCompositePanelLayerTreeItemRow::GetItemOperation)
+					SNew(STextBlock)
+					.Text(this, &SCompositePanelLayerTreeItemRow::GetItemOperation)
+					.ColorAndOpacity(FSlateColor::UseSubduedForeground())
 				];
 		}
 		
@@ -2072,6 +2076,10 @@ void SCompositePanelLayerTree::OnLayerAdded(const UClass* InLayerClass)
 	
 	CompositeActor->CompositeLayers.Insert(NewLayer, NewLayerIndex);
 
+	// Auto-increment display name if a sibling of the same type already exists
+	const FString UniqueName = UCompositePassBase::MakeUniqueDisplayName(NewLayer->GetDisplayName(), MakeArrayView(CompositeActor->CompositeLayers), NewLayer);
+	NewLayer->SetDisplayName(UniqueName);
+
 	FEditPropertyChain EditChain;
 	EditChain.AddTail(CompositeLayersProperty);
 	EditChain.SetActivePropertyNode(CompositeLayersProperty);
diff --git a/Engine/Plugins/Compositing/Composite/Source/CompositeEditor/Private/UI/SCompositePassTree.cpp b/Engine/Plugins/Compositing/Composite/Source/CompositeEditor/Private/UI/SCompositePassTree.cpp
--- a/Engine/Plugins/Compositing/Composite/Source/CompositeEditor/Private/UI/SCompositePassTree.cpp
+++ b/Engine/Plugins/Compositing/Composite/Source/CompositeEditor/Private/UI/SCompositePassTree.cpp
@@ -627,7 +627,9 @@ private:
 				.BorderImage(FAppStyle::GetBrush("NoBorder"))
 				.Padding(4.5f, 0.0f)
 				[
-					SNew(STextBlock).Text(this, &SCompositePassTreeItemRow::GetItemType)
+					SNew(STextBlock)
+					.Text(this, &SCompositePassTreeItemRow::GetItemType)
+					.ColorAndOpacity(FSlateColor::UseSubduedForeground())
 				];
 		}
 		
@@ -1693,7 +1695,16 @@ void SCompositePassTree::OnPassAdded(const UClass* InPassClass)
 	
 	FScopedTransaction AddPassTransaction(LOCTEXT("AddPassTransaction", "Add Pass"));
 	const int32 NewPassIndex = PassListOwner->AddPass(InPassClass, GroupIndex, PassIndex);
-	
+
+	// Auto-increment display name if a sibling of the same type already exists
+	if (UCompositePassBase* NewPass = PassListOwner->GetPass(GroupIndex, NewPassIndex))
+	{
+		TArray<TObjectPtr<UCompositePassBase>>& Passes = PassListOwner->GetPassesForGroup(GroupIndex);
+		const FString UniqueName = UCompositePassBase::MakeUniqueDisplayName(NewPass->GetDisplayName(), MakeArrayView(Passes), NewPass);
+		
+		NewPass->SetDisplayName(UniqueName);
+	}
+
 	if (Toolbar.IsValid())
 	{
 		Toolbar->ClearFilters();
