Composure: Refactor GetProxy() such that active passes and layers are expected to return a valid proxy.

* Remove the ambiguous bool return value from GetProxy() on all layers and passes. The method now returns FCompositeCorePassProxy* directly instead of writing to an out-param and returning true/false.
* Validity checks that previously caused early `return false` in GetProxy() are moved to IsActive() overrides on SceneCapture, PlanarReflection, ShadowReflection, and SingleLightShadow layers. Redundant IsActive() guards inside Distortion and Material pass GetProxy() implementations are removed.

#Synced-CL 52008728

diff --git a/Engine/Plugins/Compositing/Composite/Source/Composite/Private/CompositeActor.cpp b/Engine/Plugins/Compositing/Composite/Source/Composite/Private/CompositeActor.cpp
--- a/Engine/Plugins/Compositing/Composite/Source/Composite/Private/CompositeActor.cpp
+++ b/Engine/Plugins/Compositing/Composite/Source/Composite/Private/CompositeActor.cpp
@@ -106,14 +106,14 @@ namespace UE
 
 				for (const TObjectPtr<UCompositeLayerBase>& Layer : InLayers)
 				{
-					if (!IsValid(Layer) || !Layer->IsEnabled())
+					if (!IsValid(Layer) || !Layer->IsActive())
 					{
 						continue;
 					}
 
-					FCompositeCorePassProxy* OutProxy;
 					UE::CompositeCore::FPassInputDecl LayerInputDecl;
-					if (Layer->GetProxy(LayerInputDecl, Context, InFrameAllocator, OutProxy))
+					FCompositeCorePassProxy* OutProxy = Layer->GetProxy(LayerInputDecl, Context, InFrameAllocator);
+					if (OutProxy != nullptr)
 					{
 						OutProxies.Add(OutProxy);
 						Context.bIsFirstPass = false;
diff --git a/Engine/Plugins/Compositing/Composite/Source/Composite/Private/Layers/CompositeLayerBase.cpp b/Engine/Plugins/Compositing/Composite/Source/Composite/Private/Layers/CompositeLayerBase.cpp
--- a/Engine/Plugins/Compositing/Composite/Source/Composite/Private/Layers/CompositeLayerBase.cpp
+++ b/Engine/Plugins/Compositing/Composite/Source/Composite/Private/Layers/CompositeLayerBase.cpp
@@ -76,13 +76,12 @@ void UCompositeLayerBase::PostEditChangeProperty(FPropertyChangedEvent& Property
 }
 #endif //WITH_EDITOR
 
-bool UCompositeLayerBase::GetProxy(
+FCompositeCorePassProxy* UCompositeLayerBase::GetProxy(
 	const UE::CompositeCore::FPassInputDecl& /*InputDecl*/,
 	FCompositeTraversalContext& /*InContext*/,
-	FSceneRenderingBulkObjectAllocator& /*InFrameAllocator*/,
-	FCompositeCorePassProxy*& /*OutProxy*/) const
+	FSceneRenderingBulkObjectAllocator& /*InFrameAllocator*/) const
 {
-	return false;
+	return nullptr;
 }
 
 void UCompositeLayerBase::UpdateSceneCaptureComponents(USceneCaptureComponent2D& SceneCaptureComponent, TArrayView<TSoftObjectPtr<AActor>> InActors)
@@ -145,8 +144,8 @@ void UCompositeLayerBase::AddChildPasses(UE::CompositeCore::FPassInputDecl& InBa
 
 		if (IsValid(SubPass) && SubPass->IsActive())
 		{
-			FCompositeCorePassProxy* ChildPassProxy;
-			if (SubPass->GetProxy(InBasePassInput, InContext, InFrameAllocator, ChildPassProxy))
+			FCompositeCorePassProxy* ChildPassProxy = SubPass->GetProxy(InBasePassInput, InContext, InFrameAllocator);
+			if (ChildPassProxy != nullptr)
 			{
 				// Next pass should receive the output of the current child pass.
 				InBasePassInput.Set<const FCompositeCorePassProxy*>(ChildPassProxy);
diff --git a/Engine/Plugins/Compositing/Composite/Source/Composite/Private/Layers/CompositeLayerMainRender.cpp b/Engine/Plugins/Compositing/Composite/Source/Composite/Private/Layers/CompositeLayerMainRender.cpp
--- a/Engine/Plugins/Compositing/Composite/Source/Composite/Private/Layers/CompositeLayerMainRender.cpp
+++ b/Engine/Plugins/Compositing/Composite/Source/Composite/Private/Layers/CompositeLayerMainRender.cpp
@@ -12,7 +12,7 @@ UCompositeLayerMainRender::UCompositeLayerMainRender(const FObjectInitializer& O
 
 UCompositeLayerMainRender::~UCompositeLayerMainRender() = default;
 
-bool UCompositeLayerMainRender::GetProxy(const UE::CompositeCore::FPassInputDecl& /*InputDecl*/, FCompositeTraversalContext& InContext, FSceneRenderingBulkObjectAllocator& InFrameAllocator, FCompositeCorePassProxy*& OutProxy) const
+FCompositeCorePassProxy* UCompositeLayerMainRender::GetProxy(const UE::CompositeCore::FPassInputDecl& /*InputDecl*/, FCompositeTraversalContext& InContext, FSceneRenderingBulkObjectAllocator& InFrameAllocator) const
 {
 	using namespace UE::CompositeCore;
 
@@ -31,7 +31,6 @@ bool UCompositeLayerMainRender::GetProxy(const UE::CompositeCore::FPassInputDecl
 	Inputs[1] = GetDefaultSecondInput(InContext);
 
 	InContext.bNeedsSceneTextures = true;
-	OutProxy = InFrameAllocator.Create<FMergePassProxy>(MoveTemp(Inputs), GetMergeOperation(InContext), TEXT("MainRender"));
-	return true;
+	return InFrameAllocator.Create<FMergePassProxy>(MoveTemp(Inputs), GetMergeOperation(InContext), TEXT("MainRender"));
 }
 
diff --git a/Engine/Plugins/Compositing/Composite/Source/Composite/Private/Layers/CompositeLayerPlanarReflection.cpp b/Engine/Plugins/Compositing/Composite/Source/Composite/Private/Layers/CompositeLayerPlanarReflection.cpp
--- a/Engine/Plugins/Compositing/Composite/Source/Composite/Private/Layers/CompositeLayerPlanarReflection.cpp
+++ b/Engine/Plugins/Compositing/Composite/Source/Composite/Private/Layers/CompositeLayerPlanarReflection.cpp
@@ -268,21 +268,18 @@ void UCompositeLayerPlanarReflection::OnEndOfFrameUpdate(UWorld* InWorld)
 	}
 }
 
-bool UCompositeLayerPlanarReflection::GetProxy(const UE::CompositeCore::FPassInputDecl& /*InputDecl*/, FCompositeTraversalContext& InContext, FSceneRenderingBulkObjectAllocator& InFrameAllocator, FCompositeCorePassProxy*& OutProxy) const
+bool UCompositeLayerPlanarReflection::IsActive() const
 {
-	using namespace UE::CompositeCore;
+	return Super::IsActive()
+		&& IsValid(GetTypedOuter<ACompositeActor>())
+		&& IsValid(GetSceneCaptureComponent());
+}
 
-	ACompositeActor* CompositeActor = GetTypedOuter<ACompositeActor>();
-	if (!IsValid(CompositeActor))
-	{
-		return false;
-	}
+FCompositeCorePassProxy* UCompositeLayerPlanarReflection::GetProxy(const UE::CompositeCore::FPassInputDecl& /*InputDecl*/, FCompositeTraversalContext& InContext, FSceneRenderingBulkObjectAllocator& InFrameAllocator) const
+{
+	using namespace UE::CompositeCore;
 
 	USceneCaptureComponent2D* CaptureComponent = GetSceneCaptureComponent();
-	if (!IsValid(CaptureComponent))
-	{
-		return false;
-	}
 
 	FCompositeRenderTargetPool::Get().ConditionalAcquireTarget(CaptureComponent, CaptureComponent->TextureTarget, RenderTargetResolution);
 
@@ -307,8 +304,7 @@ bool UCompositeLayerPlanarReflection::GetProxy(const UE::CompositeCore::FPassInp
 	Inputs[0] = PassInput;
 	Inputs[1] = GetDefaultSecondInput(InContext);
 
-	OutProxy = InFrameAllocator.Create<FMergePassProxy>(MoveTemp(Inputs), GetMergeOperation(InContext), TEXT("PlanarReflection"), ELensDistortionHandling::Disabled);
-	return true;
+	return InFrameAllocator.Create<FMergePassProxy>(MoveTemp(Inputs), GetMergeOperation(InContext), TEXT("PlanarReflection"), ELensDistortionHandling::Disabled);
 }
 
 const TArray<TSoftObjectPtr<AActor>> UCompositeLayerPlanarReflection::GetActors() const
diff --git a/Engine/Plugins/Compositing/Composite/Source/Composite/Private/Layers/CompositeLayerPlate.cpp b/Engine/Plugins/Compositing/Composite/Source/Composite/Private/Layers/CompositeLayerPlate.cpp
--- a/Engine/Plugins/Compositing/Composite/Source/Composite/Private/Layers/CompositeLayerPlate.cpp
+++ b/Engine/Plugins/Compositing/Composite/Source/Composite/Private/Layers/CompositeLayerPlate.cpp
@@ -530,8 +530,8 @@ UE::CompositeCore::ResourceId UCompositeLayerPlate::AddPreprocessingPasses(
 			PassInput.Set<FPassInternalResourceDesc>({});
 		}
 
-		FCompositeCorePassProxy* PrepassProxy;
-		if (Prepass->GetProxy(PassInput, InContext, InFrameAllocator, PrepassProxy))
+		FCompositeCorePassProxy* PrepassProxy = Prepass->GetProxy(PassInput, InContext, InFrameAllocator);
+		if (PrepassProxy != nullptr)
 		{
 			if (PassIndex == LastValidPassIndex)
 			{
@@ -555,7 +555,7 @@ UE::CompositeCore::ResourceId UCompositeLayerPlate::AddPreprocessingPasses(
 	return OutputTextureId;
 }
 
-bool UCompositeLayerPlate::GetProxy(const UE::CompositeCore::FPassInputDecl& /*InputDecl*/, FCompositeTraversalContext& InContext, FSceneRenderingBulkObjectAllocator& InFrameAllocator, FCompositeCorePassProxy*& OutProxy) const
+FCompositeCorePassProxy* UCompositeLayerPlate::GetProxy(const UE::CompositeCore::FPassInputDecl& /*InputDecl*/, FCompositeTraversalContext& InContext, FSceneRenderingBulkObjectAllocator& InFrameAllocator) const
 {
 	using namespace UE::CompositeCore;
 	using namespace UE::Composite::Private;
@@ -640,8 +640,7 @@ bool UCompositeLayerPlate::GetProxy(const UE::CompositeCore::FPassInputDecl& /*I
 	Inputs[0] = PassInput;
 	Inputs[1] = GetDefaultSecondInput(InContext);
 
-	OutProxy = InFrameAllocator.Create<FMergePassProxy>(MoveTemp(Inputs), GetMergeOperation(InContext), TEXT("Plate"));
-	return true;
+	return InFrameAllocator.Create<FMergePassProxy>(MoveTemp(Inputs), GetMergeOperation(InContext), TEXT("Plate"));
 }
 
 TArray<UPrimitiveComponent*> UCompositeLayerPlate::GetPrimitives() const
diff --git a/Engine/Plugins/Compositing/Composite/Source/Composite/Private/Layers/CompositeLayerProcessing.cpp b/Engine/Plugins/Compositing/Composite/Source/Composite/Private/Layers/CompositeLayerProcessing.cpp
--- a/Engine/Plugins/Compositing/Composite/Source/Composite/Private/Layers/CompositeLayerProcessing.cpp
+++ b/Engine/Plugins/Compositing/Composite/Source/Composite/Private/Layers/CompositeLayerProcessing.cpp
@@ -40,13 +40,13 @@ void UCompositeLayerProcessing::PostEditChangeProperty(FPropertyChangedEvent& Pr
 }
 #endif
 
-bool UCompositeLayerProcessing::GetProxy(const UE::CompositeCore::FPassInputDecl& /*InputDecl*/, FCompositeTraversalContext& InContext, FSceneRenderingBulkObjectAllocator& InFrameAllocator, FCompositeCorePassProxy*& OutProxy) const
+FCompositeCorePassProxy* UCompositeLayerProcessing::GetProxy(const UE::CompositeCore::FPassInputDecl& /*InputDecl*/, FCompositeTraversalContext& InContext, FSceneRenderingBulkObjectAllocator& InFrameAllocator) const
 {
 	using namespace UE::CompositeCore;
-	
+
 	// Note: We ignore the passed InputDecl for now, this will be cleaned up in a subsequence refactor.
 	FPassInputDecl BasePassInput = GetDefaultSecondInput(InContext);
-	OutProxy = nullptr;
+	FCompositeCorePassProxy* OutProxy = nullptr;
 
 	// Iterate backwards so that lower passes in the UI are executed first
 	// TODO: Flip back once we have UI customization displaying passes in the correct bottom-up order
@@ -56,8 +56,8 @@ bool UCompositeLayerProcessing::GetProxy(const UE::CompositeCore::FPassInputDecl
 
 		if (IsValid(Pass) && Pass->IsActive())
 		{
-			FCompositeCorePassProxy* PassProxy;
-			if (Pass->GetProxy(BasePassInput, InContext, InFrameAllocator, PassProxy))
+			FCompositeCorePassProxy* PassProxy = Pass->GetProxy(BasePassInput, InContext, InFrameAllocator);
+			if (PassProxy != nullptr)
 			{
 				OutProxy = PassProxy;
 
@@ -67,6 +67,6 @@ bool UCompositeLayerProcessing::GetProxy(const UE::CompositeCore::FPassInputDecl
 		}
 	}
 
-	return OutProxy != nullptr;
+	return OutProxy;
 }
 
diff --git a/Engine/Plugins/Compositing/Composite/Source/Composite/Private/Layers/CompositeLayerSceneCapture.cpp b/Engine/Plugins/Compositing/Composite/Source/Composite/Private/Layers/CompositeLayerSceneCapture.cpp
--- a/Engine/Plugins/Compositing/Composite/Source/Composite/Private/Layers/CompositeLayerSceneCapture.cpp
+++ b/Engine/Plugins/Compositing/Composite/Source/Composite/Private/Layers/CompositeLayerSceneCapture.cpp
@@ -138,21 +138,18 @@ void UCompositeLayerSceneCapture::OnEndOfFrameUpdate(UWorld* InWorld)
 	}
 }
 
-bool UCompositeLayerSceneCapture::GetProxy(const UE::CompositeCore::FPassInputDecl& /*InputDecl*/, FCompositeTraversalContext& InContext, FSceneRenderingBulkObjectAllocator& InFrameAllocator, FCompositeCorePassProxy*& OutProxy) const
+bool UCompositeLayerSceneCapture::IsActive() const
 {
-	using namespace UE::CompositeCore;
+	return Super::IsActive()
+		&& IsValid(GetTypedOuter<ACompositeActor>())
+		&& IsValid(GetSceneCaptureComponent());
+}
 
-	ACompositeActor* CompositeActor = GetTypedOuter<ACompositeActor>();
-	if (!IsValid(CompositeActor))
-	{
-		return false;
-	}
+FCompositeCorePassProxy* UCompositeLayerSceneCapture::GetProxy(const UE::CompositeCore::FPassInputDecl& /*InputDecl*/, FCompositeTraversalContext& InContext, FSceneRenderingBulkObjectAllocator& InFrameAllocator) const
+{
+	using namespace UE::CompositeCore;
 
 	USceneCaptureComponent2D* CaptureComponent = GetSceneCaptureComponent();
-	if (!IsValid(CaptureComponent))
-	{
-		return false;
-	}
 
 	FCompositeRenderTargetPool::Get().ConditionalAcquireTarget(CaptureComponent, CaptureComponent->TextureTarget, RenderTargetResolution);
 
@@ -173,8 +170,7 @@ bool UCompositeLayerSceneCapture::GetProxy(const UE::CompositeCore::FPassInputDe
 	Inputs[0] = PassInput;
 	Inputs[1] = GetDefaultSecondInput(InContext);
 
-	OutProxy = InFrameAllocator.Create<FMergePassProxy>(MoveTemp(Inputs), GetMergeOperation(InContext), TEXT("SceneCapture"), ELensDistortionHandling::Disabled);
-	return true;
+	return InFrameAllocator.Create<FMergePassProxy>(MoveTemp(Inputs), GetMergeOperation(InContext), TEXT("SceneCapture"), ELensDistortionHandling::Disabled);
 }
 
 const TArray<TSoftObjectPtr<AActor>> UCompositeLayerSceneCapture::GetActors() const
diff --git a/Engine/Plugins/Compositing/Composite/Source/Composite/Private/Layers/CompositeLayerShadowReflection.cpp b/Engine/Plugins/Compositing/Composite/Source/Composite/Private/Layers/CompositeLayerShadowReflection.cpp
--- a/Engine/Plugins/Compositing/Composite/Source/Composite/Private/Layers/CompositeLayerShadowReflection.cpp
+++ b/Engine/Plugins/Compositing/Composite/Source/Composite/Private/Layers/CompositeLayerShadowReflection.cpp
@@ -155,22 +155,26 @@ void UCompositeLayerShadowReflection::OnEndOfFrameUpdate(UWorld* InWorld)
 	}
 }
 
-bool UCompositeLayerShadowReflection::GetProxy(const UE::CompositeCore::FPassInputDecl& /*InputDecl*/, FCompositeTraversalContext& InContext, FSceneRenderingBulkObjectAllocator& InFrameAllocator, FCompositeCorePassProxy*& OutProxy) const
+bool UCompositeLayerShadowReflection::IsActive() const
+{
+	return Super::IsActive()
+		&& IsValid(GetTypedOuter<ACompositeActor>());
+}
+
+FCompositeCorePassProxy* UCompositeLayerShadowReflection::GetProxy(const UE::CompositeCore::FPassInputDecl& /*InputDecl*/, FCompositeTraversalContext& InContext, FSceneRenderingBulkObjectAllocator& InFrameAllocator) const
 {
 	using namespace UE::CompositeCore;
 
 	ACompositeActor* CompositeActor = GetTypedOuter<ACompositeActor>();
-	if (!IsValid(CompositeActor))
+	if (IsValid(CompositeActor))
 	{
-		return false;
+		CachedSceneCaptures = FindOrCreateSceneCapturePair(*CompositeActor);
 	}
 
-	CachedSceneCaptures = FindOrCreateSceneCapturePair(*CompositeActor);
-
 	// Cached scene captures are now expected to be valid, early out if not.
 	if (!ensure(CachedSceneCaptures[0].IsValid() && CachedSceneCaptures[1].IsValid()))
 	{
-		return false;
+		return nullptr;
 	}
 
 	FMergePassProxy* DivisionProxy = nullptr;
@@ -207,8 +211,7 @@ bool UCompositeLayerShadowReflection::GetProxy(const UE::CompositeCore::FPassInp
 	Inputs[0].Set<const FCompositeCorePassProxy*>(DivisionProxy);
 	Inputs[1] = GetDefaultSecondInput(InContext);
 
-	OutProxy = InFrameAllocator.Create<FMergePassProxy>(MoveTemp(Inputs), GetMergeOperation(InContext), TEXT("ShadowRefl Mul"));
-	return true;
+	return InFrameAllocator.Create<FMergePassProxy>(MoveTemp(Inputs), GetMergeOperation(InContext), TEXT("ShadowRefl Mul"));
 }
 
 const TArray<TSoftObjectPtr<AActor>> UCompositeLayerShadowReflection::GetActors() const
diff --git a/Engine/Plugins/Compositing/Composite/Source/Composite/Private/Layers/CompositeLayerSingleLightShadow.cpp b/Engine/Plugins/Compositing/Composite/Source/Composite/Private/Layers/CompositeLayerSingleLightShadow.cpp
--- a/Engine/Plugins/Compositing/Composite/Source/Composite/Private/Layers/CompositeLayerSingleLightShadow.cpp
+++ b/Engine/Plugins/Compositing/Composite/Source/Composite/Private/Layers/CompositeLayerSingleLightShadow.cpp
@@ -226,20 +226,22 @@ void UCompositeLayerSingleLightShadow::OnEndOfFrameUpdate(UWorld* InWorld)
 	}
 }
 
-bool UCompositeLayerSingleLightShadow::GetProxy(const UE::CompositeCore::FPassInputDecl& /*InputDecl*/, FCompositeTraversalContext& InContext, FSceneRenderingBulkObjectAllocator& InFrameAllocator, FCompositeCorePassProxy*& OutProxy) const
+bool UCompositeLayerSingleLightShadow::IsActive() const
+{
+	return Super::IsActive()
+		&& IsValid(GetTypedOuter<ACompositeActor>())
+		&& IsValid(Light)
+		&& ShadowStrength > 0.0f;
+}
+
+FCompositeCorePassProxy* UCompositeLayerSingleLightShadow::GetProxy(const UE::CompositeCore::FPassInputDecl& /*InputDecl*/, FCompositeTraversalContext& InContext, FSceneRenderingBulkObjectAllocator& InFrameAllocator) const
 {
 	using namespace UE::CompositeCore;
 
 	ACompositeActor* CompositeActor = GetTypedOuter<ACompositeActor>();
-	if (!IsValid(CompositeActor) || !IsValid(Light))
-	{
-		return false;
-	}
-
-	// Disable with zero strength
-	if (ShadowStrength <= 0.0f)
+	if (!IsValid(CompositeActor))
 	{
-		return false;
+		return nullptr;
 	}
 
 	if (!CachedShadowMapCapture.IsValid())
@@ -255,7 +257,7 @@ bool UCompositeLayerSingleLightShadow::GetProxy(const UE::CompositeCore::FPassIn
 	// Both cached scene captures are now expected to be valid
 	if (!CachedShadowMapCapture.IsValid() || !CachedSceneDepthCapture.IsValid())
 	{
-		return false;
+		return nullptr;
 	}
 
 	// Update properties that may be dynamic
@@ -357,8 +359,7 @@ bool UCompositeLayerSingleLightShadow::GetProxy(const UE::CompositeCore::FPassIn
 	Inputs[0].Set<const FCompositeCorePassProxy*>(InputProxy);
 	Inputs[1] = GetDefaultSecondInput(InContext);
 
-	OutProxy = InFrameAllocator.Create<FMergePassProxy>(MoveTemp(Inputs), GetMergeOperation(InContext), TEXT("SingleLightShadowCatcher"));
-	return true;
+	return InFrameAllocator.Create<FMergePassProxy>(MoveTemp(Inputs), GetMergeOperation(InContext), TEXT("SingleLightShadowCatcher"));
 }
 
 const TArray<TSoftObjectPtr<AActor>> UCompositeLayerSingleLightShadow::GetShadowCastingActors() const
diff --git a/Engine/Plugins/Compositing/Composite/Source/Composite/Private/Passes/CompositePassAlphaTransform.cpp b/Engine/Plugins/Compositing/Composite/Source/Composite/Private/Passes/CompositePassAlphaTransform.cpp
--- a/Engine/Plugins/Compositing/Composite/Source/Composite/Private/Passes/CompositePassAlphaTransform.cpp
+++ b/Engine/Plugins/Compositing/Composite/Source/Composite/Private/Passes/CompositePassAlphaTransform.cpp
@@ -87,7 +87,7 @@ UCompositePassAlphaTransform::UCompositePassAlphaTransform(const FObjectInitiali
 
 UCompositePassAlphaTransform::~UCompositePassAlphaTransform() = default;
 
-bool UCompositePassAlphaTransform::GetProxy(const UE::CompositeCore::FPassInputDecl& InputDecl, FCompositeTraversalContext& InContext, FSceneRenderingBulkObjectAllocator& InFrameAllocator, FCompositeCorePassProxy*& OutProxy) const
+FCompositeCorePassProxy* UCompositePassAlphaTransform::GetProxy(const UE::CompositeCore::FPassInputDecl& InputDecl, FCompositeTraversalContext& InContext, FSceneRenderingBulkObjectAllocator& InFrameAllocator) const
 {
 	using namespace UE::Composite::Private;
 
@@ -96,6 +96,5 @@ bool UCompositePassAlphaTransform::GetProxy(const UE::CompositeCore::FPassInputD
 	Proxy->Fade = Fade;
 	Proxy->OverrideOutputAlpha = OverrideOutputAlpha;
 
-	OutProxy = Proxy;
-	return true;
+	return Proxy;
 }
diff --git a/Engine/Plugins/Compositing/Composite/Source/Composite/Private/Passes/CompositePassBlur.cpp b/Engine/Plugins/Compositing/Composite/Source/Composite/Private/Passes/CompositePassBlur.cpp
--- a/Engine/Plugins/Compositing/Composite/Source/Composite/Private/Passes/CompositePassBlur.cpp
+++ b/Engine/Plugins/Compositing/Composite/Source/Composite/Private/Passes/CompositePassBlur.cpp
@@ -121,14 +121,13 @@ bool UCompositePassBlur::IsActive() const
 	return Super::IsActive() && Radius > 1;
 }
 
-bool UCompositePassBlur::GetProxy(const UE::CompositeCore::FPassInputDecl& InputDecl, FCompositeTraversalContext& InContext, FSceneRenderingBulkObjectAllocator& InFrameAllocator, FCompositeCorePassProxy*& OutProxy) const
+FCompositeCorePassProxy* UCompositePassBlur::GetProxy(const UE::CompositeCore::FPassInputDecl& InputDecl, FCompositeTraversalContext& InContext, FSceneRenderingBulkObjectAllocator& InFrameAllocator) const
 {
 	using namespace UE::Composite;
 
 	FCompositePassBlurProxy* Proxy = InFrameAllocator.Create<FCompositePassBlurProxy>(UE::CompositeCore::FPassInputDeclArray{ InputDecl });
 	Proxy->Radius = Radius;
 
-	OutProxy = Proxy;
-	return true;
+	return Proxy;
 }
 
diff --git a/Engine/Plugins/Compositing/Composite/Source/Composite/Private/Passes/CompositePassCenteredScale.cpp b/Engine/Plugins/Compositing/Composite/Source/Composite/Private/Passes/CompositePassCenteredScale.cpp
--- a/Engine/Plugins/Compositing/Composite/Source/Composite/Private/Passes/CompositePassCenteredScale.cpp
+++ b/Engine/Plugins/Compositing/Composite/Source/Composite/Private/Passes/CompositePassCenteredScale.cpp
@@ -106,7 +106,7 @@ UCompositePassCenteredScale::UCompositePassCenteredScale(const FObjectInitialize
 
 UCompositePassCenteredScale::~UCompositePassCenteredScale() = default;
 
-bool UCompositePassCenteredScale::GetProxy(const UE::CompositeCore::FPassInputDecl& InputDecl, FCompositeTraversalContext& InContext, FSceneRenderingBulkObjectAllocator& InFrameAllocator, FCompositeCorePassProxy*& OutProxy) const
+FCompositeCorePassProxy* UCompositePassCenteredScale::GetProxy(const UE::CompositeCore::FPassInputDecl& InputDecl, FCompositeTraversalContext& InContext, FSceneRenderingBulkObjectAllocator& InFrameAllocator) const
 {
 	using namespace UE::CompositeCore;
 
@@ -117,8 +117,7 @@ bool UCompositePassCenteredScale::GetProxy(const UE::CompositeCore::FPassInputDe
 	FCenteredScalePassProxy* Proxy = InFrameAllocator.Create<FCenteredScalePassProxy>(FPassInputDeclArray{ InputDecl });
 	Proxy->ScaleUV = ScaleUV;
 
-	OutProxy = Proxy;
-	return true;
+	return Proxy;
 }
 
 bool UCompositePassCenteredScale::IsActive() const
diff --git a/Engine/Plugins/Compositing/Composite/Source/Composite/Private/Passes/CompositePassColorGrading.cpp b/Engine/Plugins/Compositing/Composite/Source/Composite/Private/Passes/CompositePassColorGrading.cpp
--- a/Engine/Plugins/Compositing/Composite/Source/Composite/Private/Passes/CompositePassColorGrading.cpp
+++ b/Engine/Plugins/Compositing/Composite/Source/Composite/Private/Passes/CompositePassColorGrading.cpp
@@ -124,7 +124,7 @@ UCompositePassColorGrading::UCompositePassColorGrading(const FObjectInitializer&
 
 UCompositePassColorGrading::~UCompositePassColorGrading() = default;
 
-bool UCompositePassColorGrading::GetProxy(const UE::CompositeCore::FPassInputDecl& InputDecl, FCompositeTraversalContext& InContext, FSceneRenderingBulkObjectAllocator& InFrameAllocator, FCompositeCorePassProxy*& OutProxy) const
+FCompositeCorePassProxy* UCompositePassColorGrading::GetProxy(const UE::CompositeCore::FPassInputDecl& InputDecl, FCompositeTraversalContext& InContext, FSceneRenderingBulkObjectAllocator& InFrameAllocator) const
 {
 	using namespace UE::Composite::Private;
 
@@ -133,7 +133,6 @@ bool UCompositePassColorGrading::GetProxy(const UE::CompositeCore::FPassInputDec
 	Proxy->ColorGradingSettings = ColorGradingSettings;
 	Proxy->bUnpremult = bUnpremult;
 
-	OutProxy = Proxy;
-	return true;
+	return Proxy;
 }
 
diff --git a/Engine/Plugins/Compositing/Composite/Source/Composite/Private/Passes/CompositePassColorKeyer.cpp b/Engine/Plugins/Compositing/Composite/Source/Composite/Private/Passes/CompositePassColorKeyer.cpp
--- a/Engine/Plugins/Compositing/Composite/Source/Composite/Private/Passes/CompositePassColorKeyer.cpp
+++ b/Engine/Plugins/Compositing/Composite/Source/Composite/Private/Passes/CompositePassColorKeyer.cpp
@@ -236,7 +236,7 @@ UCompositePassColorKeyer::UCompositePassColorKeyer(const FObjectInitializer& Obj
 
 UCompositePassColorKeyer::~UCompositePassColorKeyer() = default;
 
-bool UCompositePassColorKeyer::GetProxy(const UE::CompositeCore::FPassInputDecl& InputDecl, FCompositeTraversalContext& InContext, FSceneRenderingBulkObjectAllocator& InFrameAllocator, FCompositeCorePassProxy*& OutProxy) const
+FCompositeCorePassProxy* UCompositePassColorKeyer::GetProxy(const UE::CompositeCore::FPassInputDecl& InputDecl, FCompositeTraversalContext& InContext, FSceneRenderingBulkObjectAllocator& InFrameAllocator) const
 {
 	using namespace UE::CompositeCore;
 
@@ -256,8 +256,7 @@ bool UCompositePassColorKeyer::GetProxy(const UE::CompositeCore::FPassInputDecl&
 		Proxy->CleanPlateWeakPtr = CleanPlate;
 	}
 
-	OutProxy = Proxy;
-	return true;
+	return Proxy;
 }
 
 void UCompositePassColorKeyer::PostLoad()
diff --git a/Engine/Plugins/Compositing/Composite/Source/Composite/Private/Passes/CompositePassDistortion.cpp b/Engine/Plugins/Compositing/Composite/Source/Composite/Private/Passes/CompositePassDistortion.cpp
--- a/Engine/Plugins/Compositing/Composite/Source/Composite/Private/Passes/CompositePassDistortion.cpp
+++ b/Engine/Plugins/Compositing/Composite/Source/Composite/Private/Passes/CompositePassDistortion.cpp
@@ -190,18 +190,12 @@ bool UCompositePassDistortion::IsActive() const
 	return bIsDistortionActive;
 }
 
-bool UCompositePassDistortion::GetProxy(const UE::CompositeCore::FPassInputDecl& InputDecl, FCompositeTraversalContext& InContext, FSceneRenderingBulkObjectAllocator& InFrameAllocator, FCompositeCorePassProxy*& OutProxy) const
+FCompositeCorePassProxy* UCompositePassDistortion::GetProxy(const UE::CompositeCore::FPassInputDecl& InputDecl, FCompositeTraversalContext& InContext, FSceneRenderingBulkObjectAllocator& InFrameAllocator) const
 {
-	if (IsActive())
-	{
-		FDistortionCompositePassProxy* Proxy = InFrameAllocator.Create<FDistortionCompositePassProxy>(UE::CompositeCore::FPassInputDeclArray{ InputDecl });
-		Proxy->Distortion = Distortion;
-		Proxy->LensDistortionSVE = LensDistortionSVE;
-
-		OutProxy = Proxy;
-		return true;
-	}
+	FDistortionCompositePassProxy* Proxy = InFrameAllocator.Create<FDistortionCompositePassProxy>(UE::CompositeCore::FPassInputDeclArray{ InputDecl });
+	Proxy->Distortion = Distortion;
+	Proxy->LensDistortionSVE = LensDistortionSVE;
 
-	return false;
+	return Proxy;
 }
 
diff --git a/Engine/Plugins/Compositing/Composite/Source/Composite/Private/Passes/CompositePassFXAA.cpp b/Engine/Plugins/Compositing/Composite/Source/Composite/Private/Passes/CompositePassFXAA.cpp
--- a/Engine/Plugins/Compositing/Composite/Source/Composite/Private/Passes/CompositePassFXAA.cpp
+++ b/Engine/Plugins/Compositing/Composite/Source/Composite/Private/Passes/CompositePassFXAA.cpp
@@ -12,14 +12,13 @@ UCompositePassFXAA::UCompositePassFXAA(const FObjectInitializer& ObjectInitializ
 
 UCompositePassFXAA::~UCompositePassFXAA() = default;
 
-bool UCompositePassFXAA::GetProxy(const UE::CompositeCore::FPassInputDecl& InputDecl, FCompositeTraversalContext& InContext, FSceneRenderingBulkObjectAllocator& InFrameAllocator, FCompositeCorePassProxy*& OutProxy) const
+FCompositeCorePassProxy* UCompositePassFXAA::GetProxy(const UE::CompositeCore::FPassInputDecl& InputDecl, FCompositeTraversalContext& InContext, FSceneRenderingBulkObjectAllocator& InFrameAllocator) const
 {
 	using namespace UE::CompositeCore;
 
 	FFXAAPassProxy* Proxy = InFrameAllocator.Create<FFXAAPassProxy>(UE::CompositeCore::FPassInputDeclArray{ InputDecl });
 	Proxy->QualityOverride = Quality;
 
-	OutProxy = Proxy;
-	return true;
+	return Proxy;
 }
 
diff --git a/Engine/Plugins/Compositing/Composite/Source/Composite/Private/Passes/CompositePassMasking.cpp b/Engine/Plugins/Compositing/Composite/Source/Composite/Private/Passes/CompositePassMasking.cpp
--- a/Engine/Plugins/Compositing/Composite/Source/Composite/Private/Passes/CompositePassMasking.cpp
+++ b/Engine/Plugins/Compositing/Composite/Source/Composite/Private/Passes/CompositePassMasking.cpp
@@ -129,29 +129,28 @@ bool UCompositePassMasking::IsActive() const
 	return IsEnabled() && IsValid(MaskTexture);
 }
 
-bool UCompositePassMasking::GetProxy(const UE::CompositeCore::FPassInputDecl& InputDecl, FCompositeTraversalContext& InContext, FSceneRenderingBulkObjectAllocator& InFrameAllocator, FCompositeCorePassProxy*& OutProxy) const
+FCompositeCorePassProxy* UCompositePassMasking::GetProxy(const UE::CompositeCore::FPassInputDecl& InputDecl, FCompositeTraversalContext& InContext, FSceneRenderingBulkObjectAllocator& InFrameAllocator) const
 {
 	using namespace UE::Composite;
-	const bool bValidTexture = IsValid(MaskTexture);
 
-	if (ensureMsgf(bValidTexture, TEXT("Texture is expected to be valid for active pass.")))
+	if (!ensureMsgf(IsValid(MaskTexture), TEXT("Texture is expected to be valid for active pass.")))
 	{
-		UE::CompositeCore::FResourceMetadata MaskMetadata = { .DebugName = TEXT("MaskTexture") };
-		MaskMetadata.Filter = (MaskTexture->Filter == TF_Nearest) ? SF_Point : SF_Bilinear;
+		return nullptr;
+	}
 
-		ResourceId TextureId = InContext.FindOrCreateExternalTexture(MaskTexture, MaskMetadata);
+	UE::CompositeCore::FResourceMetadata MaskMetadata = { .DebugName = TEXT("MaskTexture") };
+	MaskMetadata.Filter = (MaskTexture->Filter == TF_Nearest) ? SF_Point : SF_Bilinear;
 
-		FPassInputDecl MaskingTextureInput;
-		MaskingTextureInput.Set<FPassExternalResourceDesc>({ TextureId });
+	ResourceId TextureId = InContext.FindOrCreateExternalTexture(MaskTexture, MaskMetadata);
 
-		FPassInputDeclArray Inputs = { InputDecl, MaskingTextureInput };
-		FMaskingPassProxy* Proxy = InFrameAllocator.Create<FMaskingPassProxy>(Inputs);
-		Proxy->MaskSourceChannel = static_cast<uint32>(MaskSource);
-		Proxy->bUnpremult = bUnpremult;
-		Proxy->bInvert = bInvert;
+	FPassInputDecl MaskingTextureInput;
+	MaskingTextureInput.Set<FPassExternalResourceDesc>({ TextureId });
 
-		OutProxy = Proxy;
-	}
+	FPassInputDeclArray Inputs = { InputDecl, MaskingTextureInput };
+	FMaskingPassProxy* Proxy = InFrameAllocator.Create<FMaskingPassProxy>(Inputs);
+	Proxy->MaskSourceChannel = static_cast<uint32>(MaskSource);
+	Proxy->bUnpremult = bUnpremult;
+	Proxy->bInvert = bInvert;
 
-	return bValidTexture;
+	return Proxy;
 }
diff --git a/Engine/Plugins/Compositing/Composite/Source/Composite/Private/Passes/CompositePassMaterial.cpp b/Engine/Plugins/Compositing/Composite/Source/Composite/Private/Passes/CompositePassMaterial.cpp
--- a/Engine/Plugins/Compositing/Composite/Source/Composite/Private/Passes/CompositePassMaterial.cpp
+++ b/Engine/Plugins/Compositing/Composite/Source/Composite/Private/Passes/CompositePassMaterial.cpp
@@ -59,21 +59,15 @@ bool UCompositePassMaterial::IsActive() const
 	return Super::IsActive() && IsValid(PostProcessMaterial) && PostProcessMaterial->IsPostProcessMaterial();
 }
 
-bool UCompositePassMaterial::GetProxy(const UE::CompositeCore::FPassInputDecl& InputDecl, FCompositeTraversalContext& InContext, FSceneRenderingBulkObjectAllocator& InFrameAllocator, FCompositeCorePassProxy*& OutProxy) const
+FCompositeCorePassProxy* UCompositePassMaterial::GetProxy(const UE::CompositeCore::FPassInputDecl& InputDecl, FCompositeTraversalContext& InContext, FSceneRenderingBulkObjectAllocator& InFrameAllocator) const
 {
-	if (IsActive())
-	{
-		FCompositePassMaterialProxy* Proxy = InFrameAllocator.Create<FCompositePassMaterialProxy>(UE::CompositeCore::FPassInputDeclArray{ InputDecl });
-		Proxy->Material = PostProcessMaterial;
-
-		// We conservatively mark any post-processing materials as potentially using scene textures.
-		// TBD if material inputs could be individually checked instead.
-		InContext.bNeedsSceneTextures = true;
-		
-		OutProxy = Proxy;
-		return true;
-	}
+	FCompositePassMaterialProxy* Proxy = InFrameAllocator.Create<FCompositePassMaterialProxy>(UE::CompositeCore::FPassInputDeclArray{ InputDecl });
+	Proxy->Material = PostProcessMaterial;
+
+	// We conservatively mark any post-processing materials as potentially using scene textures.
+	// TBD if material inputs could be individually checked instead.
+	InContext.bNeedsSceneTextures = true;
 
-	return false;
+	return Proxy;
 }
 
diff --git a/Engine/Plugins/Compositing/Composite/Source/Composite/Private/Passes/CompositePassOpenColorIO.cpp b/Engine/Plugins/Compositing/Composite/Source/Composite/Private/Passes/CompositePassOpenColorIO.cpp
--- a/Engine/Plugins/Compositing/Composite/Source/Composite/Private/Passes/CompositePassOpenColorIO.cpp
+++ b/Engine/Plugins/Compositing/Composite/Source/Composite/Private/Passes/CompositePassOpenColorIO.cpp
@@ -59,13 +59,10 @@ UCompositePassOpenColorIO::UCompositePassOpenColorIO(const FObjectInitializer& O
 
 UCompositePassOpenColorIO::~UCompositePassOpenColorIO() = default;
 
-bool UCompositePassOpenColorIO::GetProxy(const UE::CompositeCore::FPassInputDecl& InputDecl, FCompositeTraversalContext& InContext, FSceneRenderingBulkObjectAllocator& InFrameAllocator, FCompositeCorePassProxy*& OutProxy) const
+FCompositeCorePassProxy* UCompositePassOpenColorIO::GetProxy(const UE::CompositeCore::FPassInputDecl& InputDecl, FCompositeTraversalContext& InContext, FSceneRenderingBulkObjectAllocator& InFrameAllocator) const
 {
-	FOpenColorIOCompositePassProxy* Proxy = InFrameAllocator.Create<FOpenColorIOCompositePassProxy>(UE::CompositeCore::FPassInputDeclArray{ InputDecl },
+	return InFrameAllocator.Create<FOpenColorIOCompositePassProxy>(UE::CompositeCore::FPassInputDeclArray{ InputDecl },
 		FOpenColorIORendering::GetRenderPassResources(ColorConversionSettings, GMaxRHIFeatureLevel)
 	);
-
-	OutProxy = Proxy;
-	return true;
 }
 
diff --git a/Engine/Plugins/Compositing/Composite/Source/Composite/Private/Passes/CompositePassSMAA.cpp b/Engine/Plugins/Compositing/Composite/Source/Composite/Private/Passes/CompositePassSMAA.cpp
--- a/Engine/Plugins/Compositing/Composite/Source/Composite/Private/Passes/CompositePassSMAA.cpp
+++ b/Engine/Plugins/Compositing/Composite/Source/Composite/Private/Passes/CompositePassSMAA.cpp
@@ -81,7 +81,7 @@ UCompositePassSMAA::UCompositePassSMAA(const FObjectInitializer& ObjectInitializ
 
 UCompositePassSMAA::~UCompositePassSMAA() = default;
 
-bool UCompositePassSMAA::GetProxy(const UE::CompositeCore::FPassInputDecl& InputDecl, FCompositeTraversalContext& InContext, FSceneRenderingBulkObjectAllocator& InFrameAllocator, FCompositeCorePassProxy*& OutProxy) const
+FCompositeCorePassProxy* UCompositePassSMAA::GetProxy(const UE::CompositeCore::FPassInputDecl& InputDecl, FCompositeTraversalContext& InContext, FSceneRenderingBulkObjectAllocator& InFrameAllocator) const
 {
 	using namespace UE::Composite;
 
@@ -95,7 +95,6 @@ bool UCompositePassSMAA::GetProxy(const UE::CompositeCore::FPassInputDecl& Input
 	FSMAAPassProxy* Proxy = InFrameAllocator.Create<FSMAAPassProxy>(UE::CompositeCore::FPassInputDeclArray{ InputDecl });
 	Proxy->QualityOverride = Quality;
 
-	OutProxy = Proxy;
-	return true;
+	return Proxy;
 }
 
diff --git a/Engine/Plugins/Compositing/Composite/Source/Composite/Public/Layers/CompositeLayerBase.h b/Engine/Plugins/Compositing/Composite/Source/Composite/Public/Layers/CompositeLayerBase.h
--- a/Engine/Plugins/Compositing/Composite/Source/Composite/Public/Layers/CompositeLayerBase.h
+++ b/Engine/Plugins/Compositing/Composite/Source/Composite/Public/Layers/CompositeLayerBase.h
@@ -68,8 +68,8 @@ public:
 	*/
 	UE_API virtual void OnRemoved(ACompositeActor* LastOwner) {}
 
-	/** Getter function to override, returning layer proxies to be passed to the render thread. (Proxy objects should be allocated from the provided allocator.) */
-	UE_API virtual bool GetProxy(const UE::CompositeCore::FPassInputDecl& InputDecl, FCompositeTraversalContext& InContext, FSceneRenderingBulkObjectAllocator& InFrameAllocator, FCompositeCorePassProxy*& OutProxy) const override;
+	/** Override to return a render-thread proxy for this layer. Proxy objects should be allocated from the provided allocator. Only called when IsActive() returns true. */
+	UE_API virtual FCompositeCorePassProxy* GetProxy(const UE::CompositeCore::FPassInputDecl& InputDecl, FCompositeTraversalContext& InContext, FSceneRenderingBulkObjectAllocator& InFrameAllocator) const override;
 
 protected:
 	/** Convenience function to update scene capture primitive components, depending on the scene capture primitive render mode. */
diff --git a/Engine/Plugins/Compositing/Composite/Source/Composite/Public/Layers/CompositeLayerMainRender.h b/Engine/Plugins/Compositing/Composite/Source/Composite/Public/Layers/CompositeLayerMainRender.h
--- a/Engine/Plugins/Compositing/Composite/Source/Composite/Public/Layers/CompositeLayerMainRender.h
+++ b/Engine/Plugins/Compositing/Composite/Source/Composite/Public/Layers/CompositeLayerMainRender.h
@@ -24,7 +24,7 @@ public:
 	UE_API virtual ~UCompositeLayerMainRender();
 
 	/** Getter function to override, returning pass layer proxies to be passed to the render thread. (Proxy objects should be allocated from the provided allocator.) */
-	UE_API virtual bool GetProxy(const UE::CompositeCore::FPassInputDecl& InputDecl, FCompositeTraversalContext& InContext, FSceneRenderingBulkObjectAllocator& InFrameAllocator, FCompositeCorePassProxy*& OutProxy) const override;
+	UE_API virtual FCompositeCorePassProxy* GetProxy(const UE::CompositeCore::FPassInputDecl& InputDecl, FCompositeTraversalContext& InContext, FSceneRenderingBulkObjectAllocator& InFrameAllocator) const override;
 
 	/** Sub-passes applied onto the main render input before it is merged as this layer. */
 	UPROPERTY(EditAnywhere, BlueprintReadWrite, Instanced, Category = "Composite", meta = (DisallowedClasses = "/Script/Composite.CompositeLayerBase, /Script/Composite.CompositePassDistortion"))
diff --git a/Engine/Plugins/Compositing/Composite/Source/Composite/Public/Layers/CompositeLayerPlanarReflection.h b/Engine/Plugins/Compositing/Composite/Source/Composite/Public/Layers/CompositeLayerPlanarReflection.h
--- a/Engine/Plugins/Compositing/Composite/Source/Composite/Public/Layers/CompositeLayerPlanarReflection.h
+++ b/Engine/Plugins/Compositing/Composite/Source/Composite/Public/Layers/CompositeLayerPlanarReflection.h
@@ -42,8 +42,8 @@ public:
 	/** End-of-frame update callback. */
 	UE_API virtual void OnEndOfFrameUpdate(UWorld* InWorld) override;
 
-	/** Getter function to override, returning pass layer proxies to be passed to the render thread. (Proxy objects should be allocated from the provided allocator.) */
-	UE_API virtual bool GetProxy(const UE::CompositeCore::FPassInputDecl& InputDecl, FCompositeTraversalContext& InContext, FSceneRenderingBulkObjectAllocator& InFrameAllocator, FCompositeCorePassProxy*& OutProxy) const override;
+	UE_API virtual bool IsActive() const override;
+	UE_API virtual FCompositeCorePassProxy* GetProxy(const UE::CompositeCore::FPassInputDecl& InputDecl, FCompositeTraversalContext& InContext, FSceneRenderingBulkObjectAllocator& InFrameAllocator) const override;
 
 	/** Get reflection & shadow caster actors. */
 	UFUNCTION(BlueprintGetter)
diff --git a/Engine/Plugins/Compositing/Composite/Source/Composite/Public/Layers/CompositeLayerPlate.h b/Engine/Plugins/Compositing/Composite/Source/Composite/Public/Layers/CompositeLayerPlate.h
--- a/Engine/Plugins/Compositing/Composite/Source/Composite/Public/Layers/CompositeLayerPlate.h
+++ b/Engine/Plugins/Compositing/Composite/Source/Composite/Public/Layers/CompositeLayerPlate.h
@@ -69,7 +69,7 @@ public:
 	//~ End FTickableGameObject interface
 
 	/** Getter function to override, returning pass layer proxies to be passed to the render thread. (Proxy objects should be allocated from the provided allocator.) */
-	UE_API virtual bool GetProxy(const UE::CompositeCore::FPassInputDecl& InputDecl, FCompositeTraversalContext& InContext, FSceneRenderingBulkObjectAllocator& InFrameAllocator, FCompositeCorePassProxy*& OutProxy) const override;
+	UE_API virtual FCompositeCorePassProxy* GetProxy(const UE::CompositeCore::FPassInputDecl& InputDecl, FCompositeTraversalContext& InContext, FSceneRenderingBulkObjectAllocator& InFrameAllocator) const override;
 
 	/** Get array of all composite mesh primitives. */
 	UE_API TArray<UPrimitiveComponent*> GetPrimitives() const;
diff --git a/Engine/Plugins/Compositing/Composite/Source/Composite/Public/Layers/CompositeLayerProcessing.h b/Engine/Plugins/Compositing/Composite/Source/Composite/Public/Layers/CompositeLayerProcessing.h
--- a/Engine/Plugins/Compositing/Composite/Source/Composite/Public/Layers/CompositeLayerProcessing.h
+++ b/Engine/Plugins/Compositing/Composite/Source/Composite/Public/Layers/CompositeLayerProcessing.h
@@ -24,7 +24,7 @@ public:
 	UE_API virtual ~UCompositeLayerProcessing();
 
 	/** Getter function to override, returning pass layer proxies to be passed to the render thread. (Proxy objects should be allocated from the provided allocator.) */
-	UE_API virtual bool GetProxy(const UE::CompositeCore::FPassInputDecl& InputDecl, FCompositeTraversalContext& InContext, FSceneRenderingBulkObjectAllocator& InFrameAllocator, FCompositeCorePassProxy*& OutProxy) const override;
+	UE_API virtual FCompositeCorePassProxy* GetProxy(const UE::CompositeCore::FPassInputDecl& InputDecl, FCompositeTraversalContext& InContext, FSceneRenderingBulkObjectAllocator& InFrameAllocator) const override;
 
 #if WITH_EDITOR
 	UE_API virtual bool CanEditChange(const FProperty* InProperty) const;
diff --git a/Engine/Plugins/Compositing/Composite/Source/Composite/Public/Layers/CompositeLayerSceneCapture.h b/Engine/Plugins/Compositing/Composite/Source/Composite/Public/Layers/CompositeLayerSceneCapture.h
--- a/Engine/Plugins/Compositing/Composite/Source/Composite/Public/Layers/CompositeLayerSceneCapture.h
+++ b/Engine/Plugins/Compositing/Composite/Source/Composite/Public/Layers/CompositeLayerSceneCapture.h
@@ -44,8 +44,8 @@ public:
 	/** End-of-frame update callback. */
 	UE_API virtual void OnEndOfFrameUpdate(UWorld* InWorld) override;
 
-	/** Getter function to override, returning pass layer proxies to be passed to the render thread. (Proxy objects should be allocated from the provided allocator.) */
-	UE_API virtual bool GetProxy(const UE::CompositeCore::FPassInputDecl& InputDecl, FCompositeTraversalContext& InContext, FSceneRenderingBulkObjectAllocator& InFrameAllocator, FCompositeCorePassProxy*& OutProxy) const override;
+	UE_API virtual bool IsActive() const override;
+	UE_API virtual FCompositeCorePassProxy* GetProxy(const UE::CompositeCore::FPassInputDecl& InputDecl, FCompositeTraversalContext& InContext, FSceneRenderingBulkObjectAllocator& InFrameAllocator) const override;
 
 	/** Get reflection & shadow caster actors. */
 	UFUNCTION(BlueprintGetter)
diff --git a/Engine/Plugins/Compositing/Composite/Source/Composite/Public/Layers/CompositeLayerShadowReflection.h b/Engine/Plugins/Compositing/Composite/Source/Composite/Public/Layers/CompositeLayerShadowReflection.h
--- a/Engine/Plugins/Compositing/Composite/Source/Composite/Public/Layers/CompositeLayerShadowReflection.h
+++ b/Engine/Plugins/Compositing/Composite/Source/Composite/Public/Layers/CompositeLayerShadowReflection.h
@@ -78,8 +78,8 @@ public:
 	/** End-of-frame update callback. */
 	UE_API virtual void OnEndOfFrameUpdate(UWorld* InWorld) override;
 
-	/** Getter function to override, returning pass layer proxies to be passed to the render thread. (Proxy objects should be allocated from the provided allocator.) */
-	UE_API virtual bool GetProxy(const UE::CompositeCore::FPassInputDecl& InputDecl, FCompositeTraversalContext& InContext, FSceneRenderingBulkObjectAllocator& InFrameAllocator, FCompositeCorePassProxy*& OutProxy) const override;
+	UE_API virtual bool IsActive() const override;
+	UE_API virtual FCompositeCorePassProxy* GetProxy(const UE::CompositeCore::FPassInputDecl& InputDecl, FCompositeTraversalContext& InContext, FSceneRenderingBulkObjectAllocator& InFrameAllocator) const override;
 
 	/** Get the render target resolution for the scene captures. */
 	UFUNCTION(BlueprintGetter)
diff --git a/Engine/Plugins/Compositing/Composite/Source/Composite/Public/Layers/CompositeLayerSingleLightShadow.h b/Engine/Plugins/Compositing/Composite/Source/Composite/Public/Layers/CompositeLayerSingleLightShadow.h
--- a/Engine/Plugins/Compositing/Composite/Source/Composite/Public/Layers/CompositeLayerSingleLightShadow.h
+++ b/Engine/Plugins/Compositing/Composite/Source/Composite/Public/Layers/CompositeLayerSingleLightShadow.h
@@ -54,8 +54,8 @@ public:
 	/** End-of-frame update callback. */
 	UE_API virtual void OnEndOfFrameUpdate(UWorld* InWorld) override;
 
-	/** Getter function to override, returning pass layer proxies to be passed to the render thread. (Proxy objects should be allocated from the provided allocator.) */
-	UE_API virtual bool GetProxy(const UE::CompositeCore::FPassInputDecl& InputDecl, FCompositeTraversalContext& InContext, FSceneRenderingBulkObjectAllocator& InFrameAllocator, FCompositeCorePassProxy*& OutProxy) const override;
+	UE_API virtual bool IsActive() const override;
+	UE_API virtual FCompositeCorePassProxy* GetProxy(const UE::CompositeCore::FPassInputDecl& InputDecl, FCompositeTraversalContext& InContext, FSceneRenderingBulkObjectAllocator& InFrameAllocator) const override;
 
 public:
 	/**
diff --git a/Engine/Plugins/Compositing/Composite/Source/Composite/Public/Passes/CompositePassAlphaTransform.h b/Engine/Plugins/Compositing/Composite/Source/Composite/Public/Passes/CompositePassAlphaTransform.h
--- a/Engine/Plugins/Compositing/Composite/Source/Composite/Public/Passes/CompositePassAlphaTransform.h
+++ b/Engine/Plugins/Compositing/Composite/Source/Composite/Public/Passes/CompositePassAlphaTransform.h
@@ -39,7 +39,7 @@ public:
 	/** Destructor */
 	UE_API ~UCompositePassAlphaTransform();
 
-	UE_API virtual bool GetProxy(const UE::CompositeCore::FPassInputDecl& InputDecl, FCompositeTraversalContext& InContext, FSceneRenderingBulkObjectAllocator& InFrameAllocator, FCompositeCorePassProxy*& OutProxy) const override;
+	UE_API virtual FCompositeCorePassProxy* GetProxy(const UE::CompositeCore::FPassInputDecl& InputDecl, FCompositeTraversalContext& InContext, FSceneRenderingBulkObjectAllocator& InFrameAllocator) const override;
 
 public:
 	/** Alpha premultiplication operation. */
diff --git a/Engine/Plugins/Compositing/Composite/Source/Composite/Public/Passes/CompositePassBase.h b/Engine/Plugins/Compositing/Composite/Source/Composite/Public/Passes/CompositePassBase.h
--- a/Engine/Plugins/Compositing/Composite/Source/Composite/Public/Passes/CompositePassBase.h
+++ b/Engine/Plugins/Compositing/Composite/Source/Composite/Public/Passes/CompositePassBase.h
@@ -62,8 +62,8 @@ public:
 	UE_API virtual bool CanEditChange(const FProperty* InProperty) const;
 #endif
 
-	/** Getter function to override, returning pass proxies to be passed to the render thread. (Proxy objects should be allocated from the provided allocator.) */
-	UE_API virtual bool GetProxy(const UE::CompositeCore::FPassInputDecl& InputDecl, FCompositeTraversalContext& InContext, FSceneRenderingBulkObjectAllocator& InFrameAllocator, FCompositeCorePassProxy*& OutProxy) const { return false; };
+	/** Override to return a render-thread proxy for this pass. Proxy objects should be allocated from the provided allocator. Only called when IsActive() returns true. */
+	UE_API virtual FCompositeCorePassProxy* GetProxy(const UE::CompositeCore::FPassInputDecl& InputDecl, FCompositeTraversalContext& InContext, FSceneRenderingBulkObjectAllocator& InFrameAllocator) const { return nullptr; }
 
 	/** Get the enabled state. */
 	UFUNCTION(BlueprintGetter)
@@ -73,7 +73,7 @@ public:
 	UFUNCTION(BlueprintSetter)
 	UE_API virtual void SetEnabled(bool bInEnabled);
 
-	/** Returns true if the pass is returns a valid proxy. */
+	/** Returns true if the pass should participate in rendering this frame. Checked before GetProxy() is called. */
 	virtual bool IsActive() const { return IsEnabled(); }
 
 #if WITH_EDITOR
diff --git a/Engine/Plugins/Compositing/Composite/Source/Composite/Public/Passes/CompositePassBlur.h b/Engine/Plugins/Compositing/Composite/Source/Composite/Public/Passes/CompositePassBlur.h
--- a/Engine/Plugins/Compositing/Composite/Source/Composite/Public/Passes/CompositePassBlur.h
+++ b/Engine/Plugins/Compositing/Composite/Source/Composite/Public/Passes/CompositePassBlur.h
@@ -30,7 +30,7 @@ public:
 
 	UE_API virtual bool IsActive() const override;
 
-	UE_API virtual bool GetProxy(const UE::CompositeCore::FPassInputDecl& InputDecl, FCompositeTraversalContext& InContext, FSceneRenderingBulkObjectAllocator& InFrameAllocator, FCompositeCorePassProxy*& OutProxy) const override;
+	UE_API virtual FCompositeCorePassProxy* GetProxy(const UE::CompositeCore::FPassInputDecl& InputDecl, FCompositeTraversalContext& InContext, FSceneRenderingBulkObjectAllocator& InFrameAllocator) const override;
 
 public:
     /** Blur pass method. */
diff --git a/Engine/Plugins/Compositing/Composite/Source/Composite/Public/Passes/CompositePassCenteredScale.h b/Engine/Plugins/Compositing/Composite/Source/Composite/Public/Passes/CompositePassCenteredScale.h
--- a/Engine/Plugins/Compositing/Composite/Source/Composite/Public/Passes/CompositePassCenteredScale.h
+++ b/Engine/Plugins/Compositing/Composite/Source/Composite/Public/Passes/CompositePassCenteredScale.h
@@ -55,7 +55,7 @@ public:
 	/** Destructor */
 	UE_API ~UCompositePassCenteredScale();
 
-	UE_API virtual bool GetProxy(const UE::CompositeCore::FPassInputDecl& InputDecl, FCompositeTraversalContext& InContext, FSceneRenderingBulkObjectAllocator& InFrameAllocator, FCompositeCorePassProxy*& OutProxy) const override;
+	UE_API virtual FCompositeCorePassProxy* GetProxy(const UE::CompositeCore::FPassInputDecl& InputDecl, FCompositeTraversalContext& InContext, FSceneRenderingBulkObjectAllocator& InFrameAllocator) const override;
 
 	UE_API virtual bool IsActive() const override;
 
diff --git a/Engine/Plugins/Compositing/Composite/Source/Composite/Public/Passes/CompositePassColorGrading.h b/Engine/Plugins/Compositing/Composite/Source/Composite/Public/Passes/CompositePassColorGrading.h
--- a/Engine/Plugins/Compositing/Composite/Source/Composite/Public/Passes/CompositePassColorGrading.h
+++ b/Engine/Plugins/Compositing/Composite/Source/Composite/Public/Passes/CompositePassColorGrading.h
@@ -49,7 +49,7 @@ public:
 	/** Destructor */
 	UE_API ~UCompositePassColorGrading();
 
-	UE_API virtual bool GetProxy(const UE::CompositeCore::FPassInputDecl& InputDecl, FCompositeTraversalContext& InContext, FSceneRenderingBulkObjectAllocator& InFrameAllocator, FCompositeCorePassProxy*& OutProxy) const override;
+	UE_API virtual FCompositeCorePassProxy* GetProxy(const UE::CompositeCore::FPassInputDecl& InputDecl, FCompositeTraversalContext& InContext, FSceneRenderingBulkObjectAllocator& InFrameAllocator) const override;
 
 public:
 	/** Color is divided by its alpha channel before processing, and multiplied again after. */
diff --git a/Engine/Plugins/Compositing/Composite/Source/Composite/Public/Passes/CompositePassColorKeyer.h b/Engine/Plugins/Compositing/Composite/Source/Composite/Public/Passes/CompositePassColorKeyer.h
--- a/Engine/Plugins/Compositing/Composite/Source/Composite/Public/Passes/CompositePassColorKeyer.h
+++ b/Engine/Plugins/Compositing/Composite/Source/Composite/Public/Passes/CompositePassColorKeyer.h
@@ -56,7 +56,7 @@ public:
 	/** Destructor */
 	UE_API ~UCompositePassColorKeyer();
 
-	UE_API virtual bool GetProxy(const UE::CompositeCore::FPassInputDecl& InputDecl, FCompositeTraversalContext& InContext, FSceneRenderingBulkObjectAllocator& InFrameAllocator, FCompositeCorePassProxy*& OutProxy) const override;
+	UE_API virtual FCompositeCorePassProxy* GetProxy(const UE::CompositeCore::FPassInputDecl& InputDecl, FCompositeTraversalContext& InContext, FSceneRenderingBulkObjectAllocator& InFrameAllocator) const override;
 
 	UE_API virtual void PostLoad() override;
 	
diff --git a/Engine/Plugins/Compositing/Composite/Source/Composite/Public/Passes/CompositePassDistortion.h b/Engine/Plugins/Compositing/Composite/Source/Composite/Public/Passes/CompositePassDistortion.h
--- a/Engine/Plugins/Compositing/Composite/Source/Composite/Public/Passes/CompositePassDistortion.h
+++ b/Engine/Plugins/Compositing/Composite/Source/Composite/Public/Passes/CompositePassDistortion.h
@@ -37,7 +37,7 @@ public:
 
 	virtual bool IsActive() const override;
 
-	UE_API virtual bool GetProxy(const UE::CompositeCore::FPassInputDecl& InputDecl, FCompositeTraversalContext& InContext, FSceneRenderingBulkObjectAllocator& InFrameAllocator, FCompositeCorePassProxy*& OutProxy) const override;
+	UE_API virtual FCompositeCorePassProxy* GetProxy(const UE::CompositeCore::FPassInputDecl& InputDecl, FCompositeTraversalContext& InContext, FSceneRenderingBulkObjectAllocator& InFrameAllocator) const override;
 
 public:
 	/** Distortion to apply on input. */
diff --git a/Engine/Plugins/Compositing/Composite/Source/Composite/Public/Passes/CompositePassFXAA.h b/Engine/Plugins/Compositing/Composite/Source/Composite/Public/Passes/CompositePassFXAA.h
--- a/Engine/Plugins/Compositing/Composite/Source/Composite/Public/Passes/CompositePassFXAA.h
+++ b/Engine/Plugins/Compositing/Composite/Source/Composite/Public/Passes/CompositePassFXAA.h
@@ -21,7 +21,7 @@ public:
 	/** Destructor */
 	UE_API ~UCompositePassFXAA();
 
-	UE_API virtual bool GetProxy(const UE::CompositeCore::FPassInputDecl& InputDecl, FCompositeTraversalContext& InContext, FSceneRenderingBulkObjectAllocator& InFrameAllocator, FCompositeCorePassProxy*& OutProxy) const override;
+	UE_API virtual FCompositeCorePassProxy* GetProxy(const UE::CompositeCore::FPassInputDecl& InputDecl, FCompositeTraversalContext& InContext, FSceneRenderingBulkObjectAllocator& InFrameAllocator) const override;
 
 public:
 
diff --git a/Engine/Plugins/Compositing/Composite/Source/Composite/Public/Passes/CompositePassMasking.h b/Engine/Plugins/Compositing/Composite/Source/Composite/Public/Passes/CompositePassMasking.h
--- a/Engine/Plugins/Compositing/Composite/Source/Composite/Public/Passes/CompositePassMasking.h
+++ b/Engine/Plugins/Compositing/Composite/Source/Composite/Public/Passes/CompositePassMasking.h
@@ -33,7 +33,7 @@ public:
 
 	UE_API virtual bool IsActive() const override;
 
-	UE_API virtual bool GetProxy(const UE::CompositeCore::FPassInputDecl& InputDecl, FCompositeTraversalContext& InContext, FSceneRenderingBulkObjectAllocator& InFrameAllocator, FCompositeCorePassProxy*& OutProxy) const override;
+	UE_API virtual FCompositeCorePassProxy* GetProxy(const UE::CompositeCore::FPassInputDecl& InputDecl, FCompositeTraversalContext& InContext, FSceneRenderingBulkObjectAllocator& InFrameAllocator) const override;
 
 public:
 	/** Masking texture. */
diff --git a/Engine/Plugins/Compositing/Composite/Source/Composite/Public/Passes/CompositePassMaterial.h b/Engine/Plugins/Compositing/Composite/Source/Composite/Public/Passes/CompositePassMaterial.h
--- a/Engine/Plugins/Compositing/Composite/Source/Composite/Public/Passes/CompositePassMaterial.h
+++ b/Engine/Plugins/Compositing/Composite/Source/Composite/Public/Passes/CompositePassMaterial.h
@@ -25,7 +25,7 @@ public:
 
 	UE_API virtual bool IsActive() const override;
 
-	UE_API virtual bool GetProxy(const UE::CompositeCore::FPassInputDecl& InputDecl, FCompositeTraversalContext& InContext, FSceneRenderingBulkObjectAllocator& InFrameAllocator, FCompositeCorePassProxy*& OutProxy) const override;
+	UE_API virtual FCompositeCorePassProxy* GetProxy(const UE::CompositeCore::FPassInputDecl& InputDecl, FCompositeTraversalContext& InContext, FSceneRenderingBulkObjectAllocator& InFrameAllocator) const override;
 
 public:
 	/** Post-process material to execute. Input0 is connected to SceneTexture's PostprocessInput0. */
diff --git a/Engine/Plugins/Compositing/Composite/Source/Composite/Public/Passes/CompositePassOpenColorIO.h b/Engine/Plugins/Compositing/Composite/Source/Composite/Public/Passes/CompositePassOpenColorIO.h
--- a/Engine/Plugins/Compositing/Composite/Source/Composite/Public/Passes/CompositePassOpenColorIO.h
+++ b/Engine/Plugins/Compositing/Composite/Source/Composite/Public/Passes/CompositePassOpenColorIO.h
@@ -24,7 +24,7 @@ public:
 	/** Destructor */
 	UE_API ~UCompositePassOpenColorIO();
 
-	UE_API virtual bool GetProxy(const UE::CompositeCore::FPassInputDecl& InputDecl, FCompositeTraversalContext& InContext, FSceneRenderingBulkObjectAllocator& InFrameAllocator, FCompositeCorePassProxy*& OutProxy) const override;
+	UE_API virtual FCompositeCorePassProxy* GetProxy(const UE::CompositeCore::FPassInputDecl& InputDecl, FCompositeTraversalContext& InContext, FSceneRenderingBulkObjectAllocator& InFrameAllocator) const override;
 
 public:
 	/** Color conversion settings. */
diff --git a/Engine/Plugins/Compositing/Composite/Source/Composite/Public/Passes/CompositePassSMAA.h b/Engine/Plugins/Compositing/Composite/Source/Composite/Public/Passes/CompositePassSMAA.h
--- a/Engine/Plugins/Compositing/Composite/Source/Composite/Public/Passes/CompositePassSMAA.h
+++ b/Engine/Plugins/Compositing/Composite/Source/Composite/Public/Passes/CompositePassSMAA.h
@@ -21,7 +21,7 @@ public:
 	/** Destructor */
 	UE_API ~UCompositePassSMAA();
 
-	UE_API virtual bool GetProxy(const UE::CompositeCore::FPassInputDecl& InputDecl, FCompositeTraversalContext& InContext, FSceneRenderingBulkObjectAllocator& InFrameAllocator, FCompositeCorePassProxy*& OutProxy) const override;
+	UE_API virtual FCompositeCorePassProxy* GetProxy(const UE::CompositeCore::FPassInputDecl& InputDecl, FCompositeTraversalContext& InContext, FSceneRenderingBulkObjectAllocator& InFrameAllocator) const override;
 
 public:
 
