Composure: Add support for referencing sequencer-spawnable cameras.

* Camera actor picker customization to allow selection of transient spawnables.
* If previously assigned but currently not loaded, [Pending Spawnable] is shown until the spawnable is resolved.
* We serialize CameraLevelSequence & CameraBindingId properties to persist spawnable camera references. More specifically, DetectAndStoreSpawnableBinding() saves a spawnable camera's sequence + binding GUID so that TryResolveSpawnableCamera() can later re-find that camera actor when the spawnable is re-created by Sequencer (e.g. after a level reload or when playback starts).

#Synced-CL 52070808

diff --git a/Engine/Plugins/Compositing/Composite/Source/Composite/Composite.Build.cs b/Engine/Plugins/Compositing/Composite/Source/Composite/Composite.Build.cs
--- a/Engine/Plugins/Compositing/Composite/Source/Composite/Composite.Build.cs
+++ b/Engine/Plugins/Compositing/Composite/Source/Composite/Composite.Build.cs
@@ -35,6 +35,8 @@ public class Composite : ModuleRules
 				"MediaFrameworkUtilities",
 				"ConcertSyncCore",
 				"VPRoles",
+				"MovieScene",
+				"LevelSequence",
 			}
 		);
 
@@ -44,6 +46,7 @@ public class Composite : ModuleRules
 				new string[]
 				{
 					"LevelEditor",
+					"Sequencer",
 					"UnrealEd",
 					"Slate",
 					"SlateCore",
diff --git a/Engine/Plugins/Compositing/Composite/Source/Composite/Private/CompositeActor.cpp b/Engine/Plugins/Compositing/Composite/Source/Composite/Private/CompositeActor.cpp
--- a/Engine/Plugins/Compositing/Composite/Source/Composite/Private/CompositeActor.cpp
+++ b/Engine/Plugins/Compositing/Composite/Source/Composite/Private/CompositeActor.cpp
@@ -32,7 +32,15 @@
 #include "UObject/ObjectSaveContext.h"
 #include "UObject/UObjectIterator.h"
 
+#include "LevelSequenceActor.h"
+#include "LevelSequencePlayer.h"
+#include "MovieScene.h"
+#include "MovieSceneObjectBindingID.h"
+
 #if WITH_EDITOR
+#include "ISequencer.h"
+#include "LevelEditorSequencerIntegration.h"
+
 #include "LevelEditor.h"
 #include "Modules/ModuleManager.h"
 #include "Misc/TransactionObjectEvent.h"
@@ -414,6 +422,23 @@ void ACompositeActor::Tick(float DeltaTime)
 		return;
 	}
 
+	// Resolve spawnable camera if the camera binding is valid but the soft pointer is stale (e.g. after level reload with sequencer spawnables).
+	// We throttle this to avoid per-frame world iteration when the sequence isn't playing yet.
+	const bool bThrottledTick = (GFrameCounter + GetUniqueID()) % 60 == 0;
+	
+	if (CameraBindingId.IsValid() && !CameraActor.IsValid() && bThrottledTick)
+	{
+		if (AActor* ResolvedCamera = TryResolveSpawnableCamera())
+		{
+			CameraActor = ResolvedCamera;
+
+			if (IsValid(ViewProjectionComponent))
+			{
+				ViewProjectionComponent->SetCameraActor(CameraActor);
+			}
+		}
+	}
+
 	if (UCompositeWorldSubsystem* Subsystem = UWorld::GetSubsystem<UCompositeWorldSubsystem>(GetWorld()))
 	{
 		Subsystem->SetRenderWork(UE::Composite::Private::GatherFrameRenderWork(*this, SceneCapturesPerLayer));
@@ -565,10 +590,7 @@ void ACompositeActor::PostEditChangeChainProperty(struct FPropertyChangedChainEv
 	}
 	else if (PropertyName == GET_MEMBER_NAME_CHECKED(ThisClass, CameraActor))
 	{
-		if (IsValid(ViewProjectionComponent))
-		{
-			ViewProjectionComponent->SetCameraActor(CameraActor);
-		}
+		SetCameraActor(CameraActor);
 	}
 	else if (PropertyName == GET_MEMBER_NAME_CHECKED(ThisClass, VPRole))
 	{
@@ -710,11 +732,25 @@ const TSoftObjectPtr<AActor>& ACompositeActor::GetCameraActor()
 void ACompositeActor::SetCameraActor(const TSoftObjectPtr<AActor>& InActor)
 {
 	CameraActor = InActor;
-	
+
 	if (IsValid(ViewProjectionComponent))
 	{
 		ViewProjectionComponent->SetCameraActor(InActor);
 	}
+
+	// Detect spawnable binding or clear it if the camera was changed/removed.
+	if (AActor* Actor = CameraActor.Get())
+	{
+		DetectAndStoreSpawnableBinding(Actor);
+	}
+	else
+	{
+		// Either the camera was cleared (IsNull) or it has a valid soft path but the
+		// actor isn't loaded yet. In both cases, invalidate any prior spawnable binding
+		// so that Tick() doesn't re-resolve a stale spawnable over the new selection.
+		CameraBindingId.Invalidate();
+		CameraLevelSequence.Reset();
+	}
 }
 
 TArray<UCompositeLayerBase*> ACompositeActor::GetCompositeLayers()
@@ -760,6 +796,164 @@ bool ACompositeActor::IsPlatePreviewActive() const
 }
 #endif
 
+void ACompositeActor::DetectAndStoreSpawnableBinding(AActor* InActor)
+{
+	CameraBindingId.Invalidate();
+	CameraLevelSequence.Reset();
+
+	if (!IsValid(InActor))
+	{
+		return;
+	}
+
+#if WITH_EDITOR
+	// Editor path: check open sequencer instances (works when scrubbing, not just playing).
+	for (TWeakPtr<ISequencer> WeakSequencer : FLevelEditorSequencerIntegration::Get().GetSequencers())
+	{
+		TSharedPtr<ISequencer> Sequencer = WeakSequencer.Pin();
+		if (!Sequencer.IsValid() || !Sequencer->FindSharedPlaybackState())
+		{
+			continue;
+		}
+
+		// Only store bindings for transient (spawnable) actors - possessables don't need re-resolution.
+		if (!InActor->HasAnyFlags(RF_Transient))
+		{
+			continue;
+		}
+
+		FGuid FoundGuid = Sequencer->FindObjectId(*InActor, Sequencer->GetFocusedTemplateID());
+		if (!FoundGuid.IsValid())
+		{
+			continue;
+		}
+
+		ULevelSequence* RootLevelSequence = Cast<ULevelSequence>(Sequencer->GetRootMovieSceneSequence());
+		if (RootLevelSequence)
+		{
+			CameraBindingId = FoundGuid;
+			CameraLevelSequence = RootLevelSequence;
+			return;
+		}
+	}
+#endif
+
+	// Runtime path: check level sequence players (works when the sequence is playing).
+	UWorld* World = GetWorld();
+	if (!IsValid(World))
+	{
+		return;
+	}
+
+	for (TActorIterator<ALevelSequenceActor> It(World); It; ++It)
+	{
+		ALevelSequenceActor* SequenceActor = *It;
+		ULevelSequencePlayer* Player = SequenceActor ? SequenceActor->GetSequencePlayer() : nullptr;
+		if (!Player || !Player->FindSharedPlaybackState())
+		{
+			continue;
+		}
+
+		const FGuid FoundGuid = Player->FindObjectId(*InActor, FMovieSceneSequenceID());
+		if (!FoundGuid.IsValid())
+		{
+			continue;
+		}
+
+		ULevelSequence* Sequence = SequenceActor->GetSequence();
+		UMovieScene* MovieScene = Sequence ? Sequence->GetMovieScene() : nullptr;
+		if (MovieScene && MovieScene->FindSpawnable(FoundGuid))
+		{
+			CameraBindingId = FoundGuid;
+			CameraLevelSequence = Sequence;
+			return;
+		}
+	}
+}
+
+AActor* ACompositeActor::TryResolveSpawnableCamera() const
+{
+	if (!CameraBindingId.IsValid() || CameraLevelSequence.IsNull())
+	{
+		return nullptr;
+	}
+
+#if WITH_EDITOR
+	// Editor path: resolve via open sequencer instances (works when scrubbing).
+	for (TWeakPtr<ISequencer> WeakSequencer : FLevelEditorSequencerIntegration::Get().GetSequencers())
+	{
+		TSharedPtr<ISequencer> Sequencer = WeakSequencer.Pin();
+		if (!Sequencer.IsValid() || !Sequencer->FindSharedPlaybackState())
+		{
+			continue;
+		}
+
+		// Verify this sequencer instance corresponds to our stored sequence to avoid cross-sequence GUID collisions.
+		ULevelSequence* RootSequence = Cast<ULevelSequence>(Sequencer->GetRootMovieSceneSequence());
+		if (RootSequence != CameraLevelSequence.Get())
+		{
+			continue;
+		}
+
+		TArrayView<TWeakObjectPtr<>> BoundObjects = Sequencer->FindObjectsInCurrentSequence(CameraBindingId);
+		for (const TWeakObjectPtr<>& WeakObj : BoundObjects)
+		{
+			if (AActor* Actor = Cast<AActor>(WeakObj.Get()))
+			{
+				return Actor;
+			}
+		}
+	}
+#endif
+
+	// Runtime path: resolve via level sequence players (works when the sequence is playing).
+	UWorld* World = GetWorld();
+	if (!IsValid(World))
+	{
+		return nullptr;
+	}
+
+	ULevelSequence* TargetSequence = CameraLevelSequence.Get();
+	if (!TargetSequence)
+	{
+		return nullptr;
+	}
+
+	for (TActorIterator<ALevelSequenceActor> It(World); It; ++It)
+	{
+		ALevelSequenceActor* SequenceActor = *It;
+		if (!SequenceActor || SequenceActor->GetSequence() != TargetSequence)
+		{
+			continue;
+		}
+
+		ULevelSequencePlayer* Player = SequenceActor->GetSequencePlayer();
+		if (!Player)
+		{
+			continue;
+		}
+
+		TSharedPtr<const UE::MovieScene::FSharedPlaybackState> PlaybackState = Player->FindSharedPlaybackState();
+		if (!PlaybackState)
+		{
+			continue;
+		}
+
+		FMovieSceneObjectBindingID TargetObjectBindingID = UE::MovieScene::FRelativeObjectBindingID(CameraBindingId);
+
+		TArrayView<TWeakObjectPtr<>> BoundObjects = TargetObjectBindingID.ResolveBoundObjects(FMovieSceneSequenceID(), PlaybackState.ToSharedRef());
+		for (const TWeakObjectPtr<>& WeakObj : BoundObjects)
+		{
+			if (AActor* Actor = Cast<AActor>(WeakObj.Get()))
+			{
+				return Actor;
+			}
+		}
+	}
+
+	return nullptr;
+}
+
 FGameplayTag ACompositeActor::GetVPRole() const
 {
 	return VPRole;
diff --git a/Engine/Plugins/Compositing/Composite/Source/Composite/Public/CompositeActor.h b/Engine/Plugins/Compositing/Composite/Source/Composite/Public/CompositeActor.h
--- a/Engine/Plugins/Compositing/Composite/Source/Composite/Public/CompositeActor.h
+++ b/Engine/Plugins/Compositing/Composite/Source/Composite/Public/CompositeActor.h
@@ -14,6 +14,7 @@
 #define UE_API COMPOSITE_API
 
 class ACameraActor;
+class ULevelSequence;
 class UCameraComponent;
 class UCompositeShadowReflectionCatcherComponent;
 class UPrimitiveComponent;
@@ -168,12 +169,30 @@ private:
 	UPROPERTY(EditAnywhere, BlueprintGetter = IsEnabled, BlueprintSetter = SetEnabled, Interp, Category = "Composite", meta = (DisplayName = "Enabled", AllowPrivateAccess = true, DisplayPriority = "0"))
 	bool bIsEnabled;
 
-public:
 private:
 	/** The primary camera used for the composite. Used by the camera view projection component for continuous view-projection matrix updates. */
 	UPROPERTY(EditAnywhere, BlueprintGetter = GetCameraActor, BlueprintSetter = SetCameraActor, Category = "Composite", meta = (AllowPrivateAccess))
 	TSoftObjectPtr<AActor> CameraActor;
 
+	/** Level sequence containing the camera spawnable binding. Used to resolve spawnable cameras across reloads. */
+	UPROPERTY()
+	TSoftObjectPtr<ULevelSequence> CameraLevelSequence;
+
+	/** Binding GUID of the camera within CameraLevelSequence. */
+	UPROPERTY()
+	FGuid CameraBindingId;
+
+	/** Attempt to resolve a spawnable camera from the stored binding info. Returns the resolved actor or nullptr. */
+	AActor* TryResolveSpawnableCamera() const;
+
+public:
+	/** Returns true if the actor has a stored spawnable camera binding. */
+	bool HasSpawnableBinding() const { return CameraBindingId.IsValid(); }
+
+	/** Detect if a camera actor is a sequencer spawnable and store its binding info. */
+	UE_API void DetectAndStoreSpawnableBinding(AActor* InActor);
+
+private:
 	/** Array of composite layers for merging or processing images. */
 	UPROPERTY(EditAnywhere, BlueprintGetter = GetCompositeLayers, BlueprintSetter = SetCompositeLayers, Instanced, Category = "Composite", meta = (DisplayName = "Layers"))
 	TArray<TObjectPtr<UCompositeLayerBase>> CompositeLayers;
diff --git a/Engine/Plugins/Compositing/Composite/Source/CompositeEditor/CompositeEditor.Build.cs b/Engine/Plugins/Compositing/Composite/Source/CompositeEditor/CompositeEditor.Build.cs
--- a/Engine/Plugins/Compositing/Composite/Source/CompositeEditor/CompositeEditor.Build.cs
+++ b/Engine/Plugins/Compositing/Composite/Source/CompositeEditor/CompositeEditor.Build.cs
@@ -31,6 +31,7 @@ public class CompositeEditor : ModuleRules
 				"Engine",
 				"InputCore",
 				"LevelEditor",
+				"LevelSequence",
 				"MediaAssets",
 				"MediaFrameworkUtilities",
 				"MovieScene",
diff --git a/Engine/Plugins/Compositing/Composite/Source/CompositeEditor/Private/Customizations/CompositeCustomizationHelpers.cpp b/Engine/Plugins/Compositing/Composite/Source/CompositeEditor/Private/Customizations/CompositeCustomizationHelpers.cpp
new file mode 100644
--- /dev/null
+++ b/Engine/Plugins/Compositing/Composite/Source/CompositeEditor/Private/Customizations/CompositeCustomizationHelpers.cpp
@@ -0,0 +1,126 @@
+// Copyright Epic Games, Inc. All Rights Reserved.
+
+#include "CompositeCustomizationHelpers.h"
+
+#include "Camera/CameraComponent.h"
+#include "CompositeActor.h"
+#include "DetailLayoutBuilder.h"
+#include "PropertyCustomizationHelpers.h"
+#include "ScopedTransaction.h"
+#include "Widgets/Input/SComboButton.h"
+#include "Widgets/Text/STextBlock.h"
+
+#define LOCTEXT_NAMESPACE "CompositeCustomizationHelpers"
+
+namespace CompositeCustomizationHelpers
+{
+
+void CustomizeCameraActorProperty(IDetailLayoutBuilder& DetailLayout, TSharedPtr<IPropertyHandle> CameraActorHandle)
+{
+	if (!CameraActorHandle.IsValid() || !CameraActorHandle->IsValidHandle())
+	{
+		return;
+	}
+
+	auto GetCurrentActor = [CameraActorHandle]() -> AActor*
+	{
+		UObject* Object = nullptr;
+		if (CameraActorHandle->GetValue(Object) == FPropertyAccess::Success)
+		{
+			return Cast<AActor>(Object);
+		}
+		return nullptr;
+	};
+
+	auto CameraFilter = FOnShouldFilterActor::CreateLambda([](const AActor* InActor)
+	{
+		// In FSceneOutlinerFilter, returning true means the item passes the filter (is shown).
+		return InActor && InActor->GetComponentByClass<UCameraComponent>();
+	});
+
+	TSharedPtr<SComboButton> PickerCombo;
+
+	DetailLayout.EditDefaultProperty(CameraActorHandle)->CustomWidget()
+		.NameContent()
+		[
+			CameraActorHandle->CreatePropertyNameWidget()
+		]
+		.ValueContent()
+		.MinDesiredWidth(250.f)
+		[
+			SAssignNew(PickerCombo, SComboButton)
+			.ContentPadding(FMargin(4.0f, 2.0f))
+			.ButtonContent()
+			[
+				SNew(STextBlock)
+				.Text_Lambda([GetCurrentActor, CameraActorHandle]()
+				{
+					if (AActor* Actor = GetCurrentActor())
+					{
+						return FText::FromString(Actor->GetActorLabel());
+					}
+					// Show a hint when a spawnable binding exists but the actor isn't spawned yet.
+					TArray<UObject*> OuterObjects;
+					CameraActorHandle->GetOuterObjects(OuterObjects);
+					for (UObject* Outer : OuterObjects)
+					{
+						if (ACompositeActor* CompositeActor = Cast<ACompositeActor>(Outer))
+						{
+							if (CompositeActor->HasSpawnableBinding())
+							{
+								return LOCTEXT("SpawnableRef", "[Pending Spawnable]");
+							}
+						}
+					}
+					return LOCTEXT("None", "None");
+				})
+				.Font(IDetailLayoutBuilder::GetDetailFont())
+			]
+		];
+
+	PickerCombo->SetOnGetMenuContent(FOnGetContent::CreateLambda(
+		[CameraActorHandle, GetCurrentActor, CameraFilter, WeakPickerCombo = PickerCombo.ToWeakPtr()]() -> TSharedRef<SWidget>
+		{
+			return PropertyCustomizationHelpers::MakeActorPickerWithMenu(
+				GetCurrentActor(),
+				/*AllowClear=*/ true,
+				/*AllowPickingLevelInstanceContent=*/ false,
+				CameraFilter,
+				FOnActorSelected::CreateLambda([CameraActorHandle](AActor* InActor)
+				{
+					if (CameraActorHandle.IsValid() && CameraActorHandle->IsValidHandle())
+					{
+						CameraActorHandle->SetValue(FAssetData(InActor));
+
+						FScopedTransaction Transaction(LOCTEXT("SetSpawnableBinding", "Set Spawnable Camera Binding"));
+
+						TArray<UObject*> OuterObjects;
+						CameraActorHandle->GetOuterObjects(OuterObjects);
+						for (UObject* Outer : OuterObjects)
+						{
+							if (ACompositeActor* CompositeActor = Cast<ACompositeActor>(Outer))
+							{
+								CompositeActor->Modify();
+								CompositeActor->DetectAndStoreSpawnableBinding(InActor);
+							}
+						}
+					}
+				}),
+				FSimpleDelegate::CreateLambda([WeakPickerCombo]()
+				{
+					if (TSharedPtr<SComboButton> Combo = WeakPickerCombo.Pin())
+					{
+						Combo->SetIsOpen(false);
+					}
+				}),
+				FSimpleDelegate(), // OnUseSelected
+				/*bDisplayUseSelected=*/ false,
+				/*bShowTransient=*/ true
+			);
+		}
+	));
+}
+
+} // namespace CompositeCustomizationHelpers
+
+#undef LOCTEXT_NAMESPACE
diff --git a/Engine/Plugins/Compositing/Composite/Source/CompositeEditor/Private/Customizations/CompositeCustomizationHelpers.h b/Engine/Plugins/Compositing/Composite/Source/CompositeEditor/Private/Customizations/CompositeCustomizationHelpers.h
--- a/Engine/Plugins/Compositing/Composite/Source/CompositeEditor/Private/Customizations/CompositeCustomizationHelpers.h
+++ b/Engine/Plugins/Compositing/Composite/Source/CompositeEditor/Private/Customizations/CompositeCustomizationHelpers.h
@@ -2,41 +2,18 @@
 
 #pragma once
 
-#include "Camera/CameraComponent.h"
-#include "DetailLayoutBuilder.h"
-#include "PropertyCustomizationHelpers.h" // SObjectPropertyEntryBox
+#include "Templates/SharedPointer.h"
+
+class IDetailLayoutBuilder;
+class IPropertyHandle;
 
 namespace CompositeCustomizationHelpers
 {
 
-/** Customize a CameraActor soft object property to show an actor picker filtered to actors with a UCameraComponent. */
-inline void CustomizeCameraActorProperty(IDetailLayoutBuilder& DetailLayout, TSharedPtr<IPropertyHandle> CameraActorHandle)
-{
-	if (CameraActorHandle.IsValid() && CameraActorHandle->IsValidHandle())
-	{
-		DetailLayout.EditDefaultProperty(CameraActorHandle)->CustomWidget()
-			.NameContent()
-			[
-				CameraActorHandle->CreatePropertyNameWidget()
-			]
-			.ValueContent()
-			.MinDesiredWidth(250.f)
-			[
-				SNew(SObjectPropertyEntryBox)
-					.IsEnabled(true)
-					.AllowedClass(AActor::StaticClass())
-					.AllowClear(true)
-					.DisplayBrowse(true)
-					.DisplayUseSelected(true)
-					.PropertyHandle(CameraActorHandle)
-					.OnShouldFilterActor(FOnShouldFilterActor::CreateLambda([](const AActor* InActor)
-						{
-							// NOTE: For actors in SObjectPropertyEntryBox, filtered means allowed (inverse for assets).
-							return InActor && InActor->GetComponentByClass<UCameraComponent>();
-						})
-					)
-			];
-	}
-}
+/**
+ * Customize a CameraActor soft object property to show an actor picker filtered to actors with a UCameraComponent.
+ * Shows transient actors (e.g. Sequencer spawnables) in the picker.
+ */
+void CustomizeCameraActorProperty(IDetailLayoutBuilder& DetailLayout, TSharedPtr<IPropertyHandle> CameraActorHandle);
 
 } // namespace CompositeCustomizationHelpers
diff --git a/Engine/Plugins/Compositing/Composite/Source/CompositeEditor/Private/Customizations/CompositeViewProjectionComponentCustomization.cpp b/Engine/Plugins/Compositing/Composite/Source/CompositeEditor/Private/Customizations/CompositeViewProjectionComponentCustomization.cpp
--- a/Engine/Plugins/Compositing/Composite/Source/CompositeEditor/Private/Customizations/CompositeViewProjectionComponentCustomization.cpp
+++ b/Engine/Plugins/Compositing/Composite/Source/CompositeEditor/Private/Customizations/CompositeViewProjectionComponentCustomization.cpp
@@ -4,6 +4,7 @@
 
 #include "CompositeCustomizationHelpers.h"
 #include "Components/CompositeViewProjectionComponent.h"
+#include "DetailLayoutBuilder.h"
 
 #define LOCTEXT_NAMESPACE "FCompositeViewProjectionComponentCustomization"
 
