Composure: Add alpha-only blur pass mode with premultiplication handling.

* Advanced property only for now: it will become more useful once we have dilation/erosion supported.

#Synced-CL 52070808

diff --git a/Engine/Plugins/Compositing/Composite/Shaders/Private/CompositeBlur.usf b/Engine/Plugins/Compositing/Composite/Shaders/Private/CompositeBlur.usf
--- a/Engine/Plugins/Compositing/Composite/Shaders/Private/CompositeBlur.usf
+++ b/Engine/Plugins/Compositing/Composite/Shaders/Private/CompositeBlur.usf
@@ -3,6 +3,8 @@
 #include "/Engine/Private/ScreenPass.ush"
 #include "/Engine/Public/Platform.ush"
 
+#include "CompositeCommon.ush"
+
 // See https://www.rastergrid.com/blog/2010/09/efficient-gaussian-blur-with-linear-sampling/
 
 Texture2D Input_Texture;
@@ -11,6 +13,15 @@ FScreenTransform SvPositionToInputTextureUV;
 SCREEN_PASS_TEXTURE_VIEWPORT(Input)
 float2 Dir;
 float Radius;
+uint bUnpremult;
+
+#if ALPHA_ONLY
+    #define BLUR_TYPE float
+    #define TO_BLUR_TYPE(x) (x).a
+#else
+    #define BLUR_TYPE float4
+    #define TO_BLUR_TYPE(x) (x)
+#endif
 
 float NormalDistribution1D(float X)
 {
@@ -28,15 +39,15 @@ float4 SampleInputTexture(float2 UV)
     return Input_Texture.SampleLevel(Input_Sampler, ClampInputUV(UV), 0.0);
 }
 
-float4 GaussianBlur(float2 UV)
+BLUR_TYPE GaussianBlur(float2 UV)
 {
     float2 Scale = Dir * Input_ViewportSizeInverse;
-    float4 Accum = 0.0;
+    BLUR_TYPE Accum = (BLUR_TYPE)0;
     float Total = 0.0;
     float InvRadius = 1.0 / max(Radius, 1e-6);
 
     // Handle center first
-    Accum += SampleInputTexture(UV);
+    Accum += TO_BLUR_TYPE(SampleInputTexture(UV));
     Total += 1.0;
 
     // Determine how many pairs we can process symmetrically with bilinear taps
@@ -54,27 +65,27 @@ float4 GaussianBlur(float2 UV)
 
         float W1 = NormalDistribution1D(D1 * InvRadius);
         float W2 = NormalDistribution1D(D2 * InvRadius);
-        
+
         float CombinedWeight = W1 + W2;
         float Offset = D1 - W2 / CombinedWeight;
         float2 UVOffset = Scale * Offset;
-        
-        Accum += SampleInputTexture(UV + UVOffset) * CombinedWeight;
-        Accum += SampleInputTexture(UV - UVOffset) * CombinedWeight;
+
+        Accum += TO_BLUR_TYPE(SampleInputTexture(UV + UVOffset)) * CombinedWeight;
+        Accum += TO_BLUR_TYPE(SampleInputTexture(UV - UVOffset)) * CombinedWeight;
         Total += CombinedWeight * 2.0;
     }
-    
+
     // Handle the leftover pixel if Radius is odd
     if ((int)Radius % 2 != 0 && Radius >= 1.0)
     {
         float W = NormalDistribution1D(1.0 * InvRadius);
         float2 UVOffset = Scale * 1.0;
-       
-        Accum += SampleInputTexture(UV + UVOffset) * W;
-        Accum += SampleInputTexture(UV - UVOffset) * W;
+
+        Accum += TO_BLUR_TYPE(SampleInputTexture(UV + UVOffset)) * W;
+        Accum += TO_BLUR_TYPE(SampleInputTexture(UV - UVOffset)) * W;
         Total += W * 2.0;
     }
-    
+
     return Accum / max(Total, 1e-6);
 }
 
@@ -85,5 +96,19 @@ void MainPS(
 {
     float2 InputUV = ApplyScreenTransform(SvPosition.xy, SvPositionToInputTextureUV);
 
-    OutColor = GaussianBlur(InputUV);
+    BLUR_TYPE Blurred = GaussianBlur(InputUV);
+
+#if ALPHA_ONLY
+    float4 Original = SampleInputTexture(InputUV);
+    if (bUnpremult)
+    {
+        OutColor = float4(Original.rgb * SafeRcp(Original.a) * Blurred, Blurred);
+    }
+    else
+    {
+        OutColor = float4(Original.rgb, Blurred);
+    }
+#else
+    OutColor = Blurred;
+#endif
 }
diff --git a/Engine/Plugins/Compositing/Composite/Source/Composite/Private/Passes/CompositePassBlur.cpp b/Engine/Plugins/Compositing/Composite/Source/Composite/Private/Passes/CompositePassBlur.cpp
--- a/Engine/Plugins/Compositing/Composite/Source/Composite/Private/Passes/CompositePassBlur.cpp
+++ b/Engine/Plugins/Compositing/Composite/Source/Composite/Private/Passes/CompositePassBlur.cpp
@@ -15,11 +15,15 @@ class FCompositePassBlurShader : public FGlobalShader
 	DECLARE_GLOBAL_SHADER(FCompositePassBlurShader);
 	SHADER_USE_PARAMETER_STRUCT(FCompositePassBlurShader, FGlobalShader);
 
+	class FAlphaOnlyDim : SHADER_PERMUTATION_BOOL("ALPHA_ONLY");
+	using FPermutationDomain = TShaderPermutationDomain<FAlphaOnlyDim>;
+
 	BEGIN_SHADER_PARAMETER_STRUCT(FParameters, )
 		SHADER_PARAMETER_STRUCT(FScreenPassTextureInput, Input)
 		SHADER_PARAMETER(FScreenTransform, SvPositionToInputTextureUV)
 		SHADER_PARAMETER(FVector2f, Dir)
 		SHADER_PARAMETER(float, Radius)
+		SHADER_PARAMETER(uint32, bUnpremult)
 		RENDER_TARGET_BINDING_SLOTS()
 	END_SHADER_PARAMETER_STRUCT()
 };
@@ -56,7 +60,10 @@ namespace UE
 				}
 
 				FGlobalShaderMap* GlobalShaderMap = GetGlobalShaderMap(InView.GetFeatureLevel());
-				
+
+				FCompositePassBlurShader::FPermutationDomain PermutationVector;
+				PermutationVector.Set<FCompositePassBlurShader::FAlphaOnlyDim>(bAlphaOnly);
+
 				{
 					FCompositePassBlurShader::FParameters* Parameters = GraphBuilder.AllocParameters<FCompositePassBlurShader::FParameters>();
 					Parameters->Input = GetScreenPassTextureInput(InputTexture, TStaticSamplerState<SF_Bilinear, AM_Clamp, AM_Clamp, AM_Clamp>::GetRHI());
@@ -65,9 +72,10 @@ namespace UE
 						FScreenTransform::ChangeTextureBasisFromTo(FScreenPassTextureViewport(InputTexture), FScreenTransform::ETextureBasis::ViewportUV, FScreenTransform::ETextureBasis::TextureUV));
 					Parameters->Dir = FVector2f(0, 1);
 					Parameters->Radius = Radius;
+					Parameters->bUnpremult = static_cast<uint32>(bUnpremult);
 					Parameters->RenderTargets[0] = IntermediateOutput.GetRenderTargetBinding();
 
-					TShaderMapRef<FCompositePassBlurShader> PixelShader(GlobalShaderMap);
+					TShaderMapRef<FCompositePassBlurShader> PixelShader(GlobalShaderMap, PermutationVector);
 					FPixelShaderUtils::AddFullscreenPass(
 						GraphBuilder,
 						GlobalShaderMap,
@@ -86,9 +94,10 @@ namespace UE
 						FScreenTransform::ChangeTextureBasisFromTo(FScreenPassTextureViewport(IntermediateOutput), FScreenTransform::ETextureBasis::ViewportUV, FScreenTransform::ETextureBasis::TextureUV));
 					Parameters->Dir = FVector2f(1, 0);
 					Parameters->Radius = Radius;
+					Parameters->bUnpremult = static_cast<uint32>(bUnpremult);
 					Parameters->RenderTargets[0] = Output.GetRenderTargetBinding();
 
-					TShaderMapRef<FCompositePassBlurShader> PixelShader(GlobalShaderMap);
+					TShaderMapRef<FCompositePassBlurShader> PixelShader(GlobalShaderMap, PermutationVector);
 					FPixelShaderUtils::AddFullscreenPass(
 						GraphBuilder,
 						GlobalShaderMap,
@@ -103,6 +112,8 @@ namespace UE
 			}
 
 			int32 Radius = 1;
+			bool bAlphaOnly = false;
+			bool bUnpremult = true;
 		};
 	}
 }
@@ -110,7 +121,9 @@ namespace UE
 UCompositePassBlur::UCompositePassBlur(const FObjectInitializer& ObjectInitializer)
 	: Super(ObjectInitializer)
 	, Method(ECompositePassBlurMethod::Gaussian)
-    , Radius(1)
+	, Radius(1)
+	, bAlphaOnly(false)
+	, bUnpremult(true)
 {
 }
 
@@ -127,7 +140,8 @@ FCompositeCorePassProxy* UCompositePassBlur::GetProxy(const UE::CompositeCore::F
 
 	FCompositePassBlurProxy* Proxy = InFrameAllocator.Create<FCompositePassBlurProxy>(UE::CompositeCore::FPassInputDeclArray{ InputDecl });
 	Proxy->Radius = Radius;
+	Proxy->bAlphaOnly = bAlphaOnly;
+	Proxy->bUnpremult = bUnpremult;
 
 	return Proxy;
 }
-
diff --git a/Engine/Plugins/Compositing/Composite/Source/Composite/Public/Passes/CompositePassBlur.h b/Engine/Plugins/Compositing/Composite/Source/Composite/Public/Passes/CompositePassBlur.h
--- a/Engine/Plugins/Compositing/Composite/Source/Composite/Public/Passes/CompositePassBlur.h
+++ b/Engine/Plugins/Compositing/Composite/Source/Composite/Public/Passes/CompositePassBlur.h
@@ -40,6 +40,14 @@ public:
 	/** Size of the blur in pixels. */
 	UPROPERTY(EditAnywhere, BlueprintReadWrite, Interp, Category = "Composite", meta = (ClampMin = "1", UIMin = "1", UIMax = "64"))
 	int32 Radius = 1;
+
+	/** Only blur the alpha channel, leaving RGB unchanged. */
+	UPROPERTY(EditAnywhere, BlueprintReadWrite, Interp, Category = "Composite", AdvancedDisplay)
+	bool bAlphaOnly = false;
+
+	/** When alpha-only, unpremultiply RGB before blur and re-premultiply after. Disable if input is not premultiplied. */
+	UPROPERTY(BlueprintReadWrite, EditAnywhere, Interp, Category = "Composite", AdvancedDisplay, meta = (EditCondition = "bAlphaOnly"))
+	bool bUnpremult = true;
 };
 
 #undef UE_API
