Composure: Allow UTexture2DDynamic use on plate layer, as requested.

#Synced-CL 52070808

diff --git a/Engine/Plugins/Compositing/Composite/Source/Composite/Private/Layers/CompositeLayerPlate.cpp b/Engine/Plugins/Compositing/Composite/Source/Composite/Private/Layers/CompositeLayerPlate.cpp
--- a/Engine/Plugins/Compositing/Composite/Source/Composite/Private/Layers/CompositeLayerPlate.cpp
+++ b/Engine/Plugins/Compositing/Composite/Source/Composite/Private/Layers/CompositeLayerPlate.cpp
@@ -10,6 +10,7 @@
 #include "CompositeRenderTargetPool.h"
 #include "CompositeWorldSubsystem.h"
 #include "Engine/Texture2D.h"
+#include "Engine/Texture2DDynamic.h"
 #include "MediaTexture.h"
 #include "EngineUtils.h"
 #include "Engine/TextureRenderTarget2D.h"
@@ -128,6 +129,10 @@ namespace UE
 				{
 					return FIntPoint(Texture2D->GetSizeX(), Texture2D->GetSizeY());
 				}
+				else if (const UTexture2DDynamic* Texture2DDynamic = Cast<UTexture2DDynamic>(InTexture))
+				{
+					return FIntPoint(Texture2DDynamic->SizeX, Texture2DDynamic->SizeY);
+				}
 				else if (const UTextureRenderTarget2D* RenderTarget2D = Cast<UTextureRenderTarget2D>(InTexture))
 				{
 					return FIntPoint(RenderTarget2D->SizeX, RenderTarget2D->SizeY);
diff --git a/Engine/Plugins/Compositing/Composite/Source/Composite/Private/Passes/CompositePassCenteredScale.cpp b/Engine/Plugins/Compositing/Composite/Source/Composite/Private/Passes/CompositePassCenteredScale.cpp
--- a/Engine/Plugins/Compositing/Composite/Source/Composite/Private/Passes/CompositePassCenteredScale.cpp
+++ b/Engine/Plugins/Compositing/Composite/Source/Composite/Private/Passes/CompositePassCenteredScale.cpp
@@ -6,6 +6,7 @@
 #include "CompositeActor.h"
 #include "Engine/TextureRenderTarget2D.h"
 #include "Engine/Texture2D.h"
+#include "Engine/Texture2DDynamic.h"
 #include "Layers/CompositeLayerPlate.h"
 #include "MediaTexture.h"
 #include "Passes/CompositeCorePassProxy.h"
@@ -174,6 +175,10 @@ FVector2f UCompositePassCenteredScale::CalculateScale() const
 				{
 					TextureSize = FIntVector2(Texture2D->GetSizeX(), Texture2D->GetSizeY());
 				}
+				else if (const UTexture2DDynamic* Texture2DDynamic = Cast<UTexture2DDynamic>(PlateLayer->Texture))
+				{
+					TextureSize = FIntVector2(Texture2DDynamic->SizeX, Texture2DDynamic->SizeY);
+				}
 				else if (const UTextureRenderTarget2D* RenderTarget2D = Cast<UTextureRenderTarget2D>(PlateLayer->Texture))
 				{
 					TextureSize = FIntVector2(RenderTarget2D->SizeX, RenderTarget2D->SizeY);
diff --git a/Engine/Plugins/Compositing/Composite/Source/Composite/Public/Layers/CompositeLayerPlate.h b/Engine/Plugins/Compositing/Composite/Source/Composite/Public/Layers/CompositeLayerPlate.h
--- a/Engine/Plugins/Compositing/Composite/Source/Composite/Public/Layers/CompositeLayerPlate.h
+++ b/Engine/Plugins/Compositing/Composite/Source/Composite/Public/Layers/CompositeLayerPlate.h
@@ -124,7 +124,7 @@ private:
 
 public:
 	/** Media texture for compositing layer. */
-	UPROPERTY(EditAnywhere, BlueprintReadWrite, Category = "Composite", meta = (AllowedClasses="/Script/Engine.Texture2D, /Script/MediaAssets.MediaTexture"))
+	UPROPERTY(EditAnywhere, BlueprintReadWrite, Category = "Composite", meta = (AllowedClasses="/Script/Engine.Texture2D, /Script/Engine.Texture2DDynamic, /Script/MediaAssets.MediaTexture"))
 	TObjectPtr<UTexture> Texture;
 
 	/**
