Composure: Minor fix to replaced Composite.Plate.ControlMeshVisibility console variable with advanced boolean property.

#Synced-CL 52070808

diff --git a/Engine/Plugins/Compositing/Composite/Source/Composite/Private/Layers/CompositeLayerPlate.cpp b/Engine/Plugins/Compositing/Composite/Source/Composite/Private/Layers/CompositeLayerPlate.cpp
--- a/Engine/Plugins/Compositing/Composite/Source/Composite/Private/Layers/CompositeLayerPlate.cpp
+++ b/Engine/Plugins/Compositing/Composite/Source/Composite/Private/Layers/CompositeLayerPlate.cpp
@@ -34,13 +34,6 @@
 
 const FLazyName UCompositeLayerPlate::CompositeTextureName = FLazyName(TEXT("CompositeTexture"));
 
-bool UCompositeLayerPlate::GControlMeshVisibility = true;
-FAutoConsoleVariableRef CVarCompositePlateControlMeshVisibility(
-	TEXT("Composite.Plate.ControlMeshVisibility"),
-	UCompositeLayerPlate::GControlMeshVisibility,
-	TEXT("When enabled, the plate layers automatically control the visibility of registered composite meshes")
-);
-
 #define LOCTEXT_NAMESPACE "Composite"
 
 namespace UE
@@ -59,8 +52,8 @@ namespace UE
 				}
 			}
 
-			/** Convenience function to handle state on removed composite meshes. Restores visibility when GControlMeshVisibility is enabled. */
-			static void HandleRemovedCompositeMeshes(TConstArrayView<TSoftObjectPtr<AActor>> InPreEditMeshes, TConstArrayView<TSoftObjectPtr<AActor>> InPostEditMeshes)
+			/** Convenience function to handle state on removed composite meshes. Restores visibility when bInControlMeshVisibility is enabled. */
+			static void HandleRemovedCompositeMeshes(TConstArrayView<TSoftObjectPtr<AActor>> InPreEditMeshes, TConstArrayView<TSoftObjectPtr<AActor>> InPostEditMeshes, bool bInControlMeshVisibility)
 			{
 				for (const TSoftObjectPtr<AActor>& PreEditMesh : InPreEditMeshes)
 				{
@@ -79,7 +72,7 @@ namespace UE
 							{
 								PrimitiveComponent->RemoveUserDataOfClass(UCompositeAssetUserData::StaticClass());
 
-								if (UCompositeLayerPlate::GControlMeshVisibility)
+								if (bInControlMeshVisibility)
 								{
 									SetCompositeMeshVisibility(*PrimitiveComponent, true);
 								}
@@ -151,6 +144,7 @@ UCompositeLayerPlate::UCompositeLayerPlate(const FObjectInitializer& ObjectIniti
 	
 	Texture = Cast<UTexture>(DefaultMediaTexture.Object);
 	bIsHoldoutEnabled = true;
+	bControlMeshVisibility = true;
 	PlateMode = ECompositePlateMode::CompositeMesh;
 	Operation = ECompositeCoreMergeOp::Over;
 
@@ -294,7 +288,7 @@ const TArray<TSoftObjectPtr<AActor>> UCompositeLayerPlate::GetCompositeMeshes()
 void UCompositeLayerPlate::SetCompositeMeshes(TArray<TSoftObjectPtr<AActor>> InCompositeMeshes)
 {
 	// First, we untrack current meshes
-	UE::Composite::Private::HandleRemovedCompositeMeshes(CompositeMeshes, InCompositeMeshes);
+	UE::Composite::Private::HandleRemovedCompositeMeshes(CompositeMeshes, InCompositeMeshes, bControlMeshVisibility);
 
 	CompositeMeshes = MoveTemp(InCompositeMeshes);
 
@@ -316,6 +310,19 @@ void UCompositeLayerPlate::SetIsHoldoutEnabled(bool bInIsHoldoutEnabled)
 	}
 }
 
+void UCompositeLayerPlate::SetControlMeshVisibility(bool bInControlMeshVisibility)
+{
+	bControlMeshVisibility = bInControlMeshVisibility;
+
+	if (!bControlMeshVisibility)
+	{
+		for (UPrimitiveComponent* Primitive : GetPrimitives())
+		{
+			UE::Composite::Private::SetCompositeMeshVisibility(*Primitive, true);
+		}
+	}
+}
+
 UTexture* UCompositeLayerPlate::GetCompositeTexture() const
 {
 	if (CachedValidScenePasses > 0)
@@ -395,7 +402,7 @@ void UCompositeLayerPlate::PostEditChangeProperty(FPropertyChangedEvent& Propert
 
 		if (bChangeWasRemoval)
 		{
-			UE::Composite::Private::HandleRemovedCompositeMeshes(PreEditCompositeMeshes, CompositeMeshes);
+			UE::Composite::Private::HandleRemovedCompositeMeshes(PreEditCompositeMeshes, CompositeMeshes, bControlMeshVisibility);
 		}
 		else
 		{
@@ -433,6 +440,10 @@ void UCompositeLayerPlate::PostEditChangeProperty(FPropertyChangedEvent& Propert
 	{
 		SetIsHoldoutEnabled(bIsHoldoutEnabled);
 	}
+	else if (PropertyName == GET_MEMBER_NAME_CHECKED(ThisClass, bControlMeshVisibility))
+	{
+		SetControlMeshVisibility(bControlMeshVisibility);
+	}
 
 	// For simplicity, we make sure to update composite meshes on any property change.
 	UpdateCompositeMeshes();
@@ -975,7 +986,7 @@ void UCompositeLayerPlate::PropagateStateChange(ECompositeStateChangeType Change
 
 		const TArray<UPrimitiveComponent*> Primitives = GetPrimitives();
 
-		if (GControlMeshVisibility)
+		if (bControlMeshVisibility)
 		{
 			for (UPrimitiveComponent* Primitive : Primitives)
 			{
@@ -1023,7 +1034,7 @@ void UCompositeLayerPlate::PropagateStateChange(ECompositeStateChangeType Change
 
 			UnregisterCompositePrimitives(Primitives, SourceWorld);
 
-			if (GControlMeshVisibility)
+			if (bControlMeshVisibility)
 			{
 				if (ChangeType == ECompositeStateChangeType::Release)
 				{
diff --git a/Engine/Plugins/Compositing/Composite/Source/Composite/Public/Layers/CompositeLayerPlate.h b/Engine/Plugins/Compositing/Composite/Source/Composite/Public/Layers/CompositeLayerPlate.h
--- a/Engine/Plugins/Compositing/Composite/Source/Composite/Public/Layers/CompositeLayerPlate.h
+++ b/Engine/Plugins/Compositing/Composite/Source/Composite/Public/Layers/CompositeLayerPlate.h
@@ -113,6 +113,17 @@ public:
 	UFUNCTION(BlueprintSetter)
 	UE_API void SetIsHoldoutEnabled(bool bInIsHoldoutEnabled);
 
+	/** Get whether plate layers automatically control the visibility of registered composite meshes. */
+	UFUNCTION(BlueprintGetter)
+	UE_API bool GetControlMeshVisibility() const { return bControlMeshVisibility; }
+
+	/**
+	 * Set whether plate layers automatically control the visibility of registered composite meshes: turned invisible
+	 * if the layer or actor is disabled or hidden. Disabling restores visibility on all registered composite meshes.
+	 */
+	UFUNCTION(BlueprintSetter)
+	UE_API void SetControlMeshVisibility(bool bInControlMeshVisibility);
+
 	/** Returns the media texture, optionally pre-processed with layer & scene-only passes into a render target. */
 	UFUNCTION(BlueprintCallable, Category = "Composite")
 	UTexture* GetCompositeTexture() const;
@@ -151,9 +162,6 @@ public:
 	UPROPERTY(EditAnywhere, BlueprintReadWrite, Instanced, Category = "Composite", meta = (DisallowedClasses = "/Script/Composite.CompositeLayerBase, /Script/Composite.CompositePassDistortion"))
 	TArray<TObjectPtr<UCompositePassBase>> LayerPasses;
 
-	/** Composite.Plate.ControlMeshVisibility console variable value. */
-	static bool GControlMeshVisibility;
-
 private:
 	/** The plate sampling & rendering mode: direct texture sampling or separate composite mesh render. */
 	UPROPERTY(EditAnywhere, BlueprintGetter = GetPlateMode, BlueprintSetter = SetPlateMode, Interp, Category = "Composite", AdvancedDisplay, meta = (AllowPrivateAccess = true, DisplayName = "Mode"))
@@ -169,6 +177,10 @@ private:
 	UPROPERTY(EditAnywhere, BlueprintGetter = GetIsHoldoutEnabled, BlueprintSetter = SetIsHoldoutEnabled, Interp, Category = "Composite", AdvancedDisplay, meta = (AllowPrivateAccess = true))
 	bool bIsHoldoutEnabled;
 
+	/** When enabled, the plate layers automatically control the visibility of registered composite meshes. Disabling restores visibility on all registered composite meshes. */
+	UPROPERTY(EditAnywhere, BlueprintGetter = GetControlMeshVisibility, BlueprintSetter = SetControlMeshVisibility, Category = "Composite", AdvancedDisplay, meta = (AllowPrivateAccess = true))
+	bool bControlMeshVisibility;
+
 private:
 	/** Condtionally register composite mesh primitives: active holdout and add to the scene view extension custom render pass. */
 	void TryRegisterCompositePrimitives(const TArray<UPrimitiveComponent*>& InPrimitives, const UWorld* SourceWorld = nullptr) const;
