Composure: Fix stale material uniforms on WPO mesh for a depth media texture.

* Calculated normals rely on "TexelSize" which is initially cached at the wrong dimensions.

#Synced-CL 52070808

diff --git a/Engine/Plugins/Compositing/Composite/Source/Composite/Composite.Build.cs b/Engine/Plugins/Compositing/Composite/Source/Composite/Composite.Build.cs
--- a/Engine/Plugins/Compositing/Composite/Source/Composite/Composite.Build.cs
+++ b/Engine/Plugins/Compositing/Composite/Source/Composite/Composite.Build.cs
@@ -32,6 +32,7 @@ public class Composite : ModuleRules
 				"Renderer",
 				"RHI",
 				"CameraCalibrationCore",
+				"MediaUtils",
 				"MediaFrameworkUtilities",
 				"ConcertSyncCore",
 				"VPRoles",
diff --git a/Engine/Plugins/Compositing/Composite/Source/Composite/Private/Components/CompositeDepthMeshComponent.cpp b/Engine/Plugins/Compositing/Composite/Source/Composite/Private/Components/CompositeDepthMeshComponent.cpp
--- a/Engine/Plugins/Compositing/Composite/Source/Composite/Private/Components/CompositeDepthMeshComponent.cpp
+++ b/Engine/Plugins/Compositing/Composite/Source/Composite/Private/Components/CompositeDepthMeshComponent.cpp
@@ -6,6 +6,7 @@
 #include "Engine/Texture.h"
 #include "Materials/Material.h"
 #include "Materials/MaterialInstanceDynamic.h"
+#include "MediaDelegates.h"
 #include "UObject/ConstructorHelpers.h"
 
 #if WITH_EDITOR
@@ -116,6 +117,77 @@ void UCompositeDepthMeshComponent::PostLoad()
 	}
 }
 
+void UCompositeDepthMeshComponent::OnRegister()
+{
+	Super::OnRegister();
+
+	if (!HasAnyFlags(RF_ClassDefaultObject))
+	{
+		UpdateMaterial();
+
+		/**
+		* Fix for incorrect normals on depth WPO mesh.
+		*
+		* When the depth texture is a media texture, its GPU resource starts as a 2x2 placeholder.
+		* The MID's uniform expression cache captures the stale texel size (0.5 x 0.5) at load time.
+		* The normal reconstruction in the material uses texel size for gradient sampling, so stale
+		* values produce incorrect normals even though WPO displacement (which only samples depth
+		* values) works correctly. Flush the cache on the first real media sample to pick up the
+		* correct texel size at the actual texture resolution.
+		*
+		* To be removed if/when we have a more general media framework solution.
+		*/
+		if (IsValid(DepthTexture) && !OnSampleRenderedHandle.IsValid())
+		{
+			OnSampleRenderedHandle = FMediaDelegates::OnSampleRendered_RenderThread.AddLambda(
+				[WeakThis = TWeakObjectPtr<UCompositeDepthMeshComponent>(this),
+				 WeakDepthTexture = TWeakObjectPtr<UTexture>(DepthTexture)]
+				(const FTextureResource* InResource, FMediaDelegates::TSharedMediaTextureSample /*InSample*/)
+				{
+					TStrongObjectPtr<UTexture> PinnedDepthTexture = WeakDepthTexture.Pin();
+					if (!PinnedDepthTexture.IsValid() || InResource != PinnedDepthTexture->GetResource())
+					{
+						return;
+					}
+
+					// One-shot: unbind immediately on the render thread to prevent re-entry.
+					if (TStrongObjectPtr<UCompositeDepthMeshComponent> Self = WeakThis.Pin())
+					{
+						FMediaDelegates::OnSampleRendered_RenderThread.Remove(Self->OnSampleRenderedHandle);
+						Self->OnSampleRenderedHandle.Reset();
+					}
+
+					AsyncTask(ENamedThreads::GameThread, [WeakThis]()
+					{
+						TStrongObjectPtr<UCompositeDepthMeshComponent> Self = WeakThis.Pin();
+						if (!Self.IsValid())
+						{
+							return;
+						}
+
+						if (UMaterialInterface* Material = Self->GetMaterial(0))
+						{
+							constexpr bool bRecreateUniformBuffer = false;
+							Material->RecacheUniformExpressions(bRecreateUniformBuffer);
+						}
+					});
+				}
+			);
+		}
+	}
+}
+
+void UCompositeDepthMeshComponent::OnUnregister()
+{
+	if (OnSampleRenderedHandle.IsValid())
+	{
+		FMediaDelegates::OnSampleRendered_RenderThread.Remove(OnSampleRenderedHandle);
+		OnSampleRenderedHandle.Reset();
+	}
+
+	Super::OnUnregister();
+}
+
 FBoxSphereBounds UCompositeDepthMeshComponent::CalcBounds(const FTransform& LocalToWorld) const
 {
 	const FVector Origin = LocalToWorld.TransformPosition(FVector::ZeroVector);
diff --git a/Engine/Plugins/Compositing/Composite/Source/Composite/Public/Components/CompositeDepthMeshComponent.h b/Engine/Plugins/Compositing/Composite/Source/Composite/Public/Components/CompositeDepthMeshComponent.h
--- a/Engine/Plugins/Compositing/Composite/Source/Composite/Public/Components/CompositeDepthMeshComponent.h
+++ b/Engine/Plugins/Compositing/Composite/Source/Composite/Public/Components/CompositeDepthMeshComponent.h
@@ -20,6 +20,11 @@ public:
 	UE_API virtual void PostLoad() override;
 	//~ End UObject interface
 
+	//~ Begin UActorComponent interface
+	UE_API virtual void OnRegister() override;
+	UE_API virtual void OnUnregister() override;
+	//~ End UActorComponent interface
+
 #if WITH_EDITOR
 	UE_API virtual void PostEditChangeProperty(FPropertyChangedEvent& PropertyChangedEvent) override;
 	UE_API virtual void PostTransacted(const FTransactionObjectEvent& InTransactionEvent) override;
@@ -40,6 +45,9 @@ private:
 	/* Convert material to MID and update its depth texture parameter. */
 	void UpdateMaterial();
 
+	/* Handle for the media texture sample rendered delegate. */
+	FDelegateHandle OnSampleRenderedHandle;
+
 #if WITH_EDITOR
 	/* Generate our static grid mesh based on the resolution parameters. */
 	void GenerateGridMesh_Editor();
