Composure: Add keyframe support and fix Sequencer property setters.

- Add FCompositeDetailKeyframeHandler for keying Interp properties from the Composure editor panel.
- Rename ACompositeActor getters/setters to Get/Set + PropertyName convention and add CallInEditor so Sequencer's FTrackInstancePropertyBindings::CacheBinding can discover and invoke them.

#Synced-CL 52267014

diff --git a/Engine/Plugins/Compositing/Composite/Source/Composite/Private/CompositeActor.cpp b/Engine/Plugins/Compositing/Composite/Source/Composite/Private/CompositeActor.cpp
--- a/Engine/Plugins/Compositing/Composite/Source/Composite/Private/CompositeActor.cpp
+++ b/Engine/Plugins/Compositing/Composite/Source/Composite/Private/CompositeActor.cpp
@@ -87,7 +87,7 @@ namespace UE
 				// Iterate backwards so that lower passes in the UI are executed first
 				for (const TObjectPtr<UCompositeLayerBase>& Layer : InLayers)
 				{
-					if (!IsValid(Layer) || !Layer->IsEnabled())
+					if (!IsValid(Layer) || !Layer->GetIsEnabled())
 					{
 						continue;
 					}
@@ -114,7 +114,7 @@ namespace UE
 
 				for (const TObjectPtr<UCompositeLayerBase>& Layer : InLayers)
 				{
-					if (!IsValid(Layer) || !Layer->IsActive())
+					if (!IsValid(Layer) || !Layer->GetIsActive())
 					{
 						continue;
 					}
@@ -242,7 +242,7 @@ namespace UE
 				for (const auto& Pair : InSceneCapturesPerLayer)
 				{
 					const TWeakObjectPtr<const UCompositeLayerBase>& Layer = Pair.Key;
-					if (Layer.IsValid() && Layer->IsEnabled())
+					if (Layer.IsValid() && Layer->GetIsEnabled())
 					{
 						for (const TObjectPtr<UCompositeSceneCapture2DComponent>& SceneCapture : Pair.Value.Components)
 						{
@@ -353,7 +353,7 @@ void ACompositeActor::PostLoad()
 	if (VPRole.IsValid())
 	{
 		// Note that if multiple actors have the same role, this could shift which actor is active on load.
-		SetActive(IsInVPRole());
+		SetIsActive(IsInVPRole());
 	}
 }
 
@@ -522,11 +522,11 @@ void ACompositeActor::PostEditChangeChainProperty(struct FPropertyChangedChainEv
 
 	if (PropertyName == GET_MEMBER_NAME_CHECKED(ThisClass, bIsActive))
 	{
-		SetActive(bIsActive);
+		SetIsActive(bIsActive);
 	}
 	else if (PropertyName == GET_MEMBER_NAME_CHECKED(ThisClass, bIsEnabled))
 	{
-		SetEnabled(bIsEnabled);
+		SetIsEnabled(bIsEnabled);
 	}
 	else if (PropertyName == GET_MEMBER_NAME_CHECKED(ThisClass, CompositeLayers))
 	{
@@ -549,7 +549,7 @@ void ACompositeActor::PostEditChangeChainProperty(struct FPropertyChangedChainEv
 			{
 				const TObjectPtr<UCompositeLayerBase>& Layer = CompositeLayers[Index];
 
-				if (!IsValid(Layer) || !Layer->IsEnabled())
+				if (!IsValid(Layer) || !Layer->GetIsEnabled())
 				{
 					continue;
 				}
@@ -624,16 +624,16 @@ void ACompositeActor::PostEditUndo()
 
 	// Deactivate the current actor if any other actor is already active
 	// This can a occur when a previously active actor deletion is undone.
-	if (IsActive())
+	if (GetIsActive())
 	{
 		if (UCompositeWorldSubsystem* Subsystem = UWorld::GetSubsystem<UCompositeWorldSubsystem>(GetWorld()))
 		{
 			for (const TWeakObjectPtr<ACompositeActor> Actor : Subsystem->GetActors())
 			{
 				ACompositeActor* OtherActor = Actor.Get();
-				if ((OtherActor != this) && IsValid(OtherActor) && OtherActor->IsActive())
+				if ((OtherActor != this) && IsValid(OtherActor) && OtherActor->GetIsActive())
 				{
-					SetActive(false);
+					SetIsActive(false);
 					break;
 				}
 			}
@@ -662,12 +662,12 @@ bool ACompositeActor::ShouldTickIfViewportsOnly() const
 }
 #endif
 
-bool ACompositeActor::IsActive() const
+bool ACompositeActor::GetIsActive() const
 {
 	return bIsActive;
 }
 
-void ACompositeActor::SetActive(bool bInActive)
+void ACompositeActor::SetIsActive(bool bInActive)
 {
 	bIsActive = bInActive;
 
@@ -697,12 +697,12 @@ void ACompositeActor::SetActive(bool bInActive)
 	}
 }
 
-bool ACompositeActor::IsEnabled() const
+bool ACompositeActor::GetIsEnabled() const
 {
 	return bIsEnabled;
 }
 
-void ACompositeActor::SetEnabled(bool bInEnabled)
+void ACompositeActor::SetIsEnabled(bool bInEnabled)
 {
 	bIsEnabled = bInEnabled;
 
@@ -963,7 +963,7 @@ void ACompositeActor::SetVPRole(FGameplayTag Value)
 {
 	VPRole = Value;
 
-	SetActive(IsInVPRole());
+	SetIsActive(IsInVPRole());
 }
 
 template<>
@@ -1111,7 +1111,7 @@ bool ACompositeActor::IsRendering() const
 	}
 #endif
 
-	return IsEnabled() && IsActive();
+	return GetIsEnabled() && GetIsActive();
 }
 
 bool ACompositeActor::IsInVPRole() const
diff --git a/Engine/Plugins/Compositing/Composite/Source/Composite/Private/Layers/CompositeLayerBase.cpp b/Engine/Plugins/Compositing/Composite/Source/Composite/Private/Layers/CompositeLayerBase.cpp
--- a/Engine/Plugins/Compositing/Composite/Source/Composite/Private/Layers/CompositeLayerBase.cpp
+++ b/Engine/Plugins/Compositing/Composite/Source/Composite/Private/Layers/CompositeLayerBase.cpp
@@ -142,7 +142,7 @@ void UCompositeLayerBase::AddChildPasses(UE::CompositeCore::FPassInputDecl& InBa
 	{
 		const TObjectPtr<UCompositePassBase>& SubPass = InPasses[PassIndex];
 
-		if (IsValid(SubPass) && SubPass->IsActive())
+		if (IsValid(SubPass) && SubPass->GetIsActive())
 		{
 			FCompositeCorePassProxy* ChildPassProxy = SubPass->GetProxy(InBasePassInput, InContext, InFrameAllocator);
 			if (ChildPassProxy != nullptr)
@@ -186,7 +186,7 @@ void UCompositeLayerBase::ConditionalOnEndOfFrameUpdate(UWorld* InWorld)
 	 * Only call the virtual end-of-frame method if our layer hasn't been disabled or removed.
 	 * This is simpler than unregistering and re-registering if the removal is undone.
 	*/
-	if (IsEnabled())
+	if (GetIsEnabled())
 	{
 		ACompositeActor* CompositeActor = GetTypedOuter<ACompositeActor>();
 
diff --git a/Engine/Plugins/Compositing/Composite/Source/Composite/Private/Layers/CompositeLayerPlanarReflection.cpp b/Engine/Plugins/Compositing/Composite/Source/Composite/Private/Layers/CompositeLayerPlanarReflection.cpp
--- a/Engine/Plugins/Compositing/Composite/Source/Composite/Private/Layers/CompositeLayerPlanarReflection.cpp
+++ b/Engine/Plugins/Compositing/Composite/Source/Composite/Private/Layers/CompositeLayerPlanarReflection.cpp
@@ -268,9 +268,9 @@ void UCompositeLayerPlanarReflection::OnEndOfFrameUpdate(UWorld* InWorld)
 	}
 }
 
-bool UCompositeLayerPlanarReflection::IsActive() const
+bool UCompositeLayerPlanarReflection::GetIsActive() const
 {
-	return Super::IsActive()
+	return Super::GetIsActive()
 		&& IsValid(GetTypedOuter<ACompositeActor>())
 		&& IsValid(GetSceneCaptureComponent());
 }
diff --git a/Engine/Plugins/Compositing/Composite/Source/Composite/Private/Layers/CompositeLayerPlate.cpp b/Engine/Plugins/Compositing/Composite/Source/Composite/Private/Layers/CompositeLayerPlate.cpp
--- a/Engine/Plugins/Compositing/Composite/Source/Composite/Private/Layers/CompositeLayerPlate.cpp
+++ b/Engine/Plugins/Compositing/Composite/Source/Composite/Private/Layers/CompositeLayerPlate.cpp
@@ -164,9 +164,9 @@ UCompositeLayerPlate::UCompositeLayerPlate(const FObjectInitializer& ObjectIniti
 
 UCompositeLayerPlate::~UCompositeLayerPlate() = default;
 
-void UCompositeLayerPlate::SetEnabled(bool bInEnabled)
+void UCompositeLayerPlate::SetIsEnabled(bool bInEnabled)
 {
-	Super::SetEnabled(bInEnabled);
+	Super::SetIsEnabled(bInEnabled);
 
 	PropagateStateChange(bInEnabled ? ECompositeStateChangeType::Activate : ECompositeStateChangeType::Deactivate, GetWorld());
 }
@@ -393,7 +393,7 @@ void UCompositeLayerPlate::PostEditChangeProperty(FPropertyChangedEvent& Propert
 
 	if (PropertyName == GET_MEMBER_NAME_CHECKED(ThisClass, bIsEnabled))
 	{
-		SetEnabled(bIsEnabled);
+		SetIsEnabled(bIsEnabled);
 	}
 	else if (PropertyName == GET_MEMBER_NAME_CHECKED(ThisClass, CompositeMeshes))
 	{
@@ -494,7 +494,7 @@ int32 UCompositeLayerPlate::FindLastValidPassIndex(TArrayView<const TObjectPtr<U
 	{
 		const TObjectPtr<UCompositePassBase>& Pass = InPasses[PassIndex];
 
-		if (IsValid(Pass) && Pass->IsActive())
+		if (IsValid(Pass) && Pass->GetIsActive())
 		{
 			LastValidPassIndex = PassIndex;
 		}
@@ -528,7 +528,7 @@ UE::CompositeCore::ResourceId UCompositeLayerPlate::AddPreprocessingPasses(
 	{
 		const TObjectPtr<UCompositePassBase>& Prepass = InPasses[PassIndex];
 
-		if (!IsValid(Prepass) || !Prepass->IsActive())
+		if (!IsValid(Prepass) || !Prepass->GetIsActive())
 		{
 			continue;
 		}
@@ -723,7 +723,7 @@ int32 UCompositeLayerPlate::GetValidPassesNum(TArrayView<const TObjectPtr<UCompo
 	{
 		const TObjectPtr<UCompositePassBase>& Pass = InPasses[Index];
 
-		if (IsValid(Pass) && Pass->IsActive())
+		if (IsValid(Pass) && Pass->GetIsActive())
 		{
 			++Count;
 		}
@@ -769,7 +769,7 @@ void UCompositeLayerPlate::TryCloseMediaProfileSource()
 
 bool UCompositeLayerPlate::IsRendering() const
 {
-	if (IsEnabled())
+	if (GetIsEnabled())
 	{
 		const ACompositeActor* CompositeActor = GetTypedOuter<ACompositeActor>();
 
diff --git a/Engine/Plugins/Compositing/Composite/Source/Composite/Private/Layers/CompositeLayerProcessing.cpp b/Engine/Plugins/Compositing/Composite/Source/Composite/Private/Layers/CompositeLayerProcessing.cpp
--- a/Engine/Plugins/Compositing/Composite/Source/Composite/Private/Layers/CompositeLayerProcessing.cpp
+++ b/Engine/Plugins/Compositing/Composite/Source/Composite/Private/Layers/CompositeLayerProcessing.cpp
@@ -54,7 +54,7 @@ FCompositeCorePassProxy* UCompositeLayerProcessing::GetProxy(const UE::Composite
 	{
 		const TObjectPtr<UCompositePassBase>& Pass = LayerPasses[PassIndex];
 
-		if (IsValid(Pass) && Pass->IsActive())
+		if (IsValid(Pass) && Pass->GetIsActive())
 		{
 			FCompositeCorePassProxy* PassProxy = Pass->GetProxy(BasePassInput, InContext, InFrameAllocator);
 			if (PassProxy != nullptr)
diff --git a/Engine/Plugins/Compositing/Composite/Source/Composite/Private/Layers/CompositeLayerSceneCapture.cpp b/Engine/Plugins/Compositing/Composite/Source/Composite/Private/Layers/CompositeLayerSceneCapture.cpp
--- a/Engine/Plugins/Compositing/Composite/Source/Composite/Private/Layers/CompositeLayerSceneCapture.cpp
+++ b/Engine/Plugins/Compositing/Composite/Source/Composite/Private/Layers/CompositeLayerSceneCapture.cpp
@@ -138,9 +138,9 @@ void UCompositeLayerSceneCapture::OnEndOfFrameUpdate(UWorld* InWorld)
 	}
 }
 
-bool UCompositeLayerSceneCapture::IsActive() const
+bool UCompositeLayerSceneCapture::GetIsActive() const
 {
-	return Super::IsActive()
+	return Super::GetIsActive()
 		&& IsValid(GetTypedOuter<ACompositeActor>())
 		&& IsValid(GetSceneCaptureComponent());
 }
diff --git a/Engine/Plugins/Compositing/Composite/Source/Composite/Private/Layers/CompositeLayerShadowReflection.cpp b/Engine/Plugins/Compositing/Composite/Source/Composite/Private/Layers/CompositeLayerShadowReflection.cpp
--- a/Engine/Plugins/Compositing/Composite/Source/Composite/Private/Layers/CompositeLayerShadowReflection.cpp
+++ b/Engine/Plugins/Compositing/Composite/Source/Composite/Private/Layers/CompositeLayerShadowReflection.cpp
@@ -43,7 +43,7 @@ void UCompositeLayerShadowReflection::PostLoad()
 {
 	Super::PostLoad();
 
-	if (IsEnabled() && !Actors.IsEmpty())
+	if (GetIsEnabled() && !Actors.IsEmpty())
 	{
 		SetHiddenPrimitiveState(true);
 	}
@@ -130,7 +130,7 @@ void UCompositeLayerShadowReflection::PostTransacted(const FTransactionObjectEve
 
 void UCompositeLayerShadowReflection::OnEndOfFrameUpdate(UWorld* InWorld)
 {
-	// Note: OnEndOfFrameUpdate is only called from the base layer if IsEnabled() is true.
+	// Note: OnEndOfFrameUpdate is only called from the base layer if GetIsEnabled() is true.
 
 	/**
 	* In order to retain the ability to reference actors in sublevels, we prefer using weak pointer components on scene captures. However,
@@ -155,9 +155,9 @@ void UCompositeLayerShadowReflection::OnEndOfFrameUpdate(UWorld* InWorld)
 	}
 }
 
-bool UCompositeLayerShadowReflection::IsActive() const
+bool UCompositeLayerShadowReflection::GetIsActive() const
 {
-	return Super::IsActive()
+	return Super::GetIsActive()
 		&& IsValid(GetTypedOuter<ACompositeActor>());
 }
 
@@ -226,7 +226,7 @@ void UCompositeLayerShadowReflection::SetActors(TArray<TSoftObjectPtr<AActor>> I
 
 	Actors = MoveTemp(InActors);
 
-	if (IsEnabled() && !Actors.IsEmpty())
+	if (GetIsEnabled() && !Actors.IsEmpty())
 	{
 		SetHiddenPrimitiveState(true);
 	}
@@ -260,9 +260,9 @@ void UCompositeLayerShadowReflection::SetRenderTargetResolution(FIntPoint InRend
 	RenderTargetResolution.Y = FMath::Max(InRenderTargetResolution.Y, 1);
 }
 
-void UCompositeLayerShadowReflection::SetEnabled(bool bInEnabled)
+void UCompositeLayerShadowReflection::SetIsEnabled(bool bInEnabled)
 {
-	Super::SetEnabled(bInEnabled);
+	Super::SetIsEnabled(bInEnabled);
 
 	SetHiddenPrimitiveState(bInEnabled && !Actors.IsEmpty());
 }
diff --git a/Engine/Plugins/Compositing/Composite/Source/Composite/Private/Layers/CompositeLayerSingleLightShadow.cpp b/Engine/Plugins/Compositing/Composite/Source/Composite/Private/Layers/CompositeLayerSingleLightShadow.cpp
--- a/Engine/Plugins/Compositing/Composite/Source/Composite/Private/Layers/CompositeLayerSingleLightShadow.cpp
+++ b/Engine/Plugins/Compositing/Composite/Source/Composite/Private/Layers/CompositeLayerSingleLightShadow.cpp
@@ -226,9 +226,9 @@ void UCompositeLayerSingleLightShadow::OnEndOfFrameUpdate(UWorld* InWorld)
 	}
 }
 
-bool UCompositeLayerSingleLightShadow::IsActive() const
+bool UCompositeLayerSingleLightShadow::GetIsActive() const
 {
-	return Super::IsActive()
+	return Super::GetIsActive()
 		&& IsValid(GetTypedOuter<ACompositeActor>())
 		&& IsValid(Light)
 		&& ShadowStrength > 0.0f;
diff --git a/Engine/Plugins/Compositing/Composite/Source/Composite/Private/Passes/CompositePassBase.cpp b/Engine/Plugins/Compositing/Composite/Source/Composite/Private/Passes/CompositePassBase.cpp
--- a/Engine/Plugins/Compositing/Composite/Source/Composite/Private/Passes/CompositePassBase.cpp
+++ b/Engine/Plugins/Compositing/Composite/Source/Composite/Private/Passes/CompositePassBase.cpp
@@ -43,12 +43,12 @@ UCompositePassBase::UCompositePassBase(const FObjectInitializer& ObjectInitializ
 
 UCompositePassBase::~UCompositePassBase() = default;
 
-bool UCompositePassBase::IsEnabled() const
+bool UCompositePassBase::GetIsEnabled() const
 {
 	return bIsEnabled;
 }
 
-void UCompositePassBase::SetEnabled(bool bInEnabled)
+void UCompositePassBase::SetIsEnabled(bool bInEnabled)
 {
 	bIsEnabled = bInEnabled;
 }
diff --git a/Engine/Plugins/Compositing/Composite/Source/Composite/Private/Passes/CompositePassBlur.cpp b/Engine/Plugins/Compositing/Composite/Source/Composite/Private/Passes/CompositePassBlur.cpp
--- a/Engine/Plugins/Compositing/Composite/Source/Composite/Private/Passes/CompositePassBlur.cpp
+++ b/Engine/Plugins/Compositing/Composite/Source/Composite/Private/Passes/CompositePassBlur.cpp
@@ -129,9 +129,9 @@ UCompositePassBlur::UCompositePassBlur(const FObjectInitializer& ObjectInitializ
 
 UCompositePassBlur::~UCompositePassBlur() = default;
 
-bool UCompositePassBlur::IsActive() const
+bool UCompositePassBlur::GetIsActive() const
 {
-	return Super::IsActive() && Radius > 1;
+	return Super::GetIsActive() && Radius > 1;
 }
 
 FCompositeCorePassProxy* UCompositePassBlur::GetProxy(const UE::CompositeCore::FPassInputDecl& InputDecl, FCompositeTraversalContext& InContext, FSceneRenderingBulkObjectAllocator& InFrameAllocator) const
diff --git a/Engine/Plugins/Compositing/Composite/Source/Composite/Private/Passes/CompositePassCenteredScale.cpp b/Engine/Plugins/Compositing/Composite/Source/Composite/Private/Passes/CompositePassCenteredScale.cpp
--- a/Engine/Plugins/Compositing/Composite/Source/Composite/Private/Passes/CompositePassCenteredScale.cpp
+++ b/Engine/Plugins/Compositing/Composite/Source/Composite/Private/Passes/CompositePassCenteredScale.cpp
@@ -111,7 +111,7 @@ FCompositeCorePassProxy* UCompositePassCenteredScale::GetProxy(const UE::Composi
 {
 	using namespace UE::CompositeCore;
 
-	// Use cached value from IsActive() if available, otherwise compute fresh.
+	// Use cached value from GetIsActive() if available, otherwise compute fresh.
 	const FVector2f ScaleUV = CachedScaleUV.IsSet() ? CachedScaleUV.GetValue() : CalculateScale();
 	CachedScaleUV.Reset();
 
@@ -121,11 +121,11 @@ FCompositeCorePassProxy* UCompositePassCenteredScale::GetProxy(const UE::Composi
 	return Proxy;
 }
 
-bool UCompositePassCenteredScale::IsActive() const
+bool UCompositePassCenteredScale::GetIsActive() const
 {
 	CachedScaleUV = CalculateScale();
 
-	return IsEnabled()
+	return GetIsEnabled()
 		&& (ScaleMode != ECompositePassScaleMode::None || OverscanUncropMode != ECompositePassOverscanUncropMode::None)
 		&& !CachedScaleUV.GetValue().Equals(FVector2f::One(), UE_SMALL_NUMBER);
 }
diff --git a/Engine/Plugins/Compositing/Composite/Source/Composite/Private/Passes/CompositePassDistortion.cpp b/Engine/Plugins/Compositing/Composite/Source/Composite/Private/Passes/CompositePassDistortion.cpp
--- a/Engine/Plugins/Compositing/Composite/Source/Composite/Private/Passes/CompositePassDistortion.cpp
+++ b/Engine/Plugins/Compositing/Composite/Source/Composite/Private/Passes/CompositePassDistortion.cpp
@@ -170,11 +170,11 @@ void UCompositePassDistortion::PostInitProperties()
 	}
 }
 
-bool UCompositePassDistortion::IsActive() const
+bool UCompositePassDistortion::GetIsActive() const
 {
 	bool bIsDistortionActive = false;
 
-	if (Super::IsActive())
+	if (Super::GetIsActive())
 	{
 		static IConsoleVariable* PaniniCVar = IConsoleManager::Get().FindConsoleVariable(TEXT("r.LensDistortion.Panini.D"));
 		if (LensDistortionSVE && LensDistortionSVE->HasDistortionState_AnyThread())
diff --git a/Engine/Plugins/Compositing/Composite/Source/Composite/Private/Passes/CompositePassMasking.cpp b/Engine/Plugins/Compositing/Composite/Source/Composite/Private/Passes/CompositePassMasking.cpp
--- a/Engine/Plugins/Compositing/Composite/Source/Composite/Private/Passes/CompositePassMasking.cpp
+++ b/Engine/Plugins/Compositing/Composite/Source/Composite/Private/Passes/CompositePassMasking.cpp
@@ -124,9 +124,9 @@ UCompositePassMasking::UCompositePassMasking(const FObjectInitializer& ObjectIni
 
 UCompositePassMasking::~UCompositePassMasking() = default;
 
-bool UCompositePassMasking::IsActive() const
+bool UCompositePassMasking::GetIsActive() const
 {
-	return IsEnabled() && IsValid(MaskTexture);
+	return GetIsEnabled() && IsValid(MaskTexture);
 }
 
 FCompositeCorePassProxy* UCompositePassMasking::GetProxy(const UE::CompositeCore::FPassInputDecl& InputDecl, FCompositeTraversalContext& InContext, FSceneRenderingBulkObjectAllocator& InFrameAllocator) const
diff --git a/Engine/Plugins/Compositing/Composite/Source/Composite/Private/Passes/CompositePassMaterial.cpp b/Engine/Plugins/Compositing/Composite/Source/Composite/Private/Passes/CompositePassMaterial.cpp
--- a/Engine/Plugins/Compositing/Composite/Source/Composite/Private/Passes/CompositePassMaterial.cpp
+++ b/Engine/Plugins/Compositing/Composite/Source/Composite/Private/Passes/CompositePassMaterial.cpp
@@ -54,9 +54,9 @@ UCompositePassMaterial::UCompositePassMaterial(const FObjectInitializer& ObjectI
 
 UCompositePassMaterial::~UCompositePassMaterial() = default;
 
-bool UCompositePassMaterial::IsActive() const
+bool UCompositePassMaterial::GetIsActive() const
 {
-	return Super::IsActive() && IsValid(PostProcessMaterial) && PostProcessMaterial->IsPostProcessMaterial();
+	return Super::GetIsActive() && IsValid(PostProcessMaterial) && PostProcessMaterial->IsPostProcessMaterial();
 }
 
 FCompositeCorePassProxy* UCompositePassMaterial::GetProxy(const UE::CompositeCore::FPassInputDecl& InputDecl, FCompositeTraversalContext& InContext, FSceneRenderingBulkObjectAllocator& InFrameAllocator) const
diff --git a/Engine/Plugins/Compositing/Composite/Source/Composite/Public/CompositeActor.h b/Engine/Plugins/Compositing/Composite/Source/Composite/Public/CompositeActor.h
--- a/Engine/Plugins/Compositing/Composite/Source/Composite/Public/CompositeActor.h
+++ b/Engine/Plugins/Compositing/Composite/Source/Composite/Public/CompositeActor.h
@@ -87,19 +87,19 @@ public:
 
 	/** Get the (locally) active state. */
 	UFUNCTION(BlueprintGetter)
-	UE_API bool IsActive() const;
+	UE_API bool GetIsActive() const;
 
 	/** Set the (locally) active state. */
 	UFUNCTION(BlueprintSetter)
-	UE_API void SetActive(bool bInActive);
+	UE_API void SetIsActive(bool bInActive);
 
 	/** Get the enabled state. */
 	UFUNCTION(BlueprintGetter)
-	UE_API bool IsEnabled() const;
+	UE_API bool GetIsEnabled() const;
 
 	/** Set the enabled state. */
-	UFUNCTION(BlueprintSetter)
-	UE_API void SetEnabled(bool bInEnabled);
+	UFUNCTION(BlueprintSetter, CallInEditor)
+	UE_API void SetIsEnabled(bool bInEnabled);
 
 	/** Camera component getter. */
 	UFUNCTION(BlueprintGetter)
@@ -114,7 +114,7 @@ public:
 	UE_API const TSoftObjectPtr<AActor>& GetCameraActor();
 
 	/** Camera actor setter. */
-	UFUNCTION(BlueprintSetter)
+	UFUNCTION(BlueprintSetter, CallInEditor)
 	UE_API void SetCameraActor(const TSoftObjectPtr<AActor>& InActor);
 
 	/** Get the composite layers. */
@@ -162,16 +162,16 @@ private:
 #endif
 
 	/** Whether or not the composite behavior is active locally - used primarily for multi-user. */
-	UPROPERTY(NonTransactional, EditInstanceOnly, BlueprintGetter = IsActive, BlueprintSetter = SetActive, Category = "Composite", meta = (AllowPrivateAccess = true))
+	UPROPERTY(NonTransactional, EditInstanceOnly, BlueprintGetter = GetIsActive, BlueprintSetter = SetIsActive, Category = "Composite", meta = (AllowPrivateAccess = true))
 	bool bIsActive;
 
 	/** Whether or not the composite behavior is enabled. */
-	UPROPERTY(EditAnywhere, BlueprintGetter = IsEnabled, BlueprintSetter = SetEnabled, Interp, Category = "Composite", meta = (DisplayName = "Enabled", AllowPrivateAccess = true, DisplayPriority = "0"))
+	UPROPERTY(EditAnywhere, BlueprintGetter = GetIsEnabled, BlueprintSetter = SetIsEnabled, Interp, Category = "Composite", meta = (DisplayName = "Enabled", AllowPrivateAccess = true, DisplayPriority = "0"))
 	bool bIsEnabled;
 
 private:
 	/** The primary camera used for the composite. Used by the camera view projection component for continuous view-projection matrix updates. */
-	UPROPERTY(EditAnywhere, BlueprintGetter = GetCameraActor, BlueprintSetter = SetCameraActor, Category = "Composite", meta = (AllowPrivateAccess))
+	UPROPERTY(EditAnywhere, BlueprintGetter = GetCameraActor, BlueprintSetter = SetCameraActor, Interp, Category = "Composite", meta = (AllowPrivateAccess))
 	TSoftObjectPtr<AActor> CameraActor;
 
 	/** Level sequence containing the camera spawnable binding. Used to resolve spawnable cameras across reloads. */
diff --git a/Engine/Plugins/Compositing/Composite/Source/Composite/Public/Layers/CompositeLayerBase.h b/Engine/Plugins/Compositing/Composite/Source/Composite/Public/Layers/CompositeLayerBase.h
--- a/Engine/Plugins/Compositing/Composite/Source/Composite/Public/Layers/CompositeLayerBase.h
+++ b/Engine/Plugins/Compositing/Composite/Source/Composite/Public/Layers/CompositeLayerBase.h
@@ -68,7 +68,7 @@ public:
 	*/
 	UE_API virtual void OnRemoved(ACompositeActor* LastOwner) {}
 
-	/** Override to return a render-thread proxy for this layer. Proxy objects should be allocated from the provided allocator. Only called when IsActive() returns true. */
+	/** Override to return a render-thread proxy for this layer. Proxy objects should be allocated from the provided allocator. Only called when GetIsActive() returns true. */
 	UE_API virtual FCompositeCorePassProxy* GetProxy(const UE::CompositeCore::FPassInputDecl& InputDecl, FCompositeTraversalContext& InContext, FSceneRenderingBulkObjectAllocator& InFrameAllocator) const override;
 
 protected:
diff --git a/Engine/Plugins/Compositing/Composite/Source/Composite/Public/Layers/CompositeLayerPlanarReflection.h b/Engine/Plugins/Compositing/Composite/Source/Composite/Public/Layers/CompositeLayerPlanarReflection.h
--- a/Engine/Plugins/Compositing/Composite/Source/Composite/Public/Layers/CompositeLayerPlanarReflection.h
+++ b/Engine/Plugins/Compositing/Composite/Source/Composite/Public/Layers/CompositeLayerPlanarReflection.h
@@ -42,7 +42,7 @@ public:
 	/** End-of-frame update callback. */
 	UE_API virtual void OnEndOfFrameUpdate(UWorld* InWorld) override;
 
-	UE_API virtual bool IsActive() const override;
+	UE_API virtual bool GetIsActive() const override;
 	UE_API virtual FCompositeCorePassProxy* GetProxy(const UE::CompositeCore::FPassInputDecl& InputDecl, FCompositeTraversalContext& InContext, FSceneRenderingBulkObjectAllocator& InFrameAllocator) const override;
 
 	/** Get reflection & shadow caster actors. */
diff --git a/Engine/Plugins/Compositing/Composite/Source/Composite/Public/Layers/CompositeLayerPlate.h b/Engine/Plugins/Compositing/Composite/Source/Composite/Public/Layers/CompositeLayerPlate.h
--- a/Engine/Plugins/Compositing/Composite/Source/Composite/Public/Layers/CompositeLayerPlate.h
+++ b/Engine/Plugins/Compositing/Composite/Source/Composite/Public/Layers/CompositeLayerPlate.h
@@ -54,7 +54,7 @@ public:
 	UE_API virtual void OnRenderingStateChange(ECompositeStateChangeType ChangeType) override;
 
 	/** Set the enabled state. */
-	UE_API virtual void SetEnabled(bool bInEnabled) override;
+	UE_API virtual void SetIsEnabled(bool bInEnabled) override;
 
 	//~ Begin FTickableGameObject interface
 	UE_API virtual void Tick(float DeltaTime) override;
diff --git a/Engine/Plugins/Compositing/Composite/Source/Composite/Public/Layers/CompositeLayerSceneCapture.h b/Engine/Plugins/Compositing/Composite/Source/Composite/Public/Layers/CompositeLayerSceneCapture.h
--- a/Engine/Plugins/Compositing/Composite/Source/Composite/Public/Layers/CompositeLayerSceneCapture.h
+++ b/Engine/Plugins/Compositing/Composite/Source/Composite/Public/Layers/CompositeLayerSceneCapture.h
@@ -44,7 +44,7 @@ public:
 	/** End-of-frame update callback. */
 	UE_API virtual void OnEndOfFrameUpdate(UWorld* InWorld) override;
 
-	UE_API virtual bool IsActive() const override;
+	UE_API virtual bool GetIsActive() const override;
 	UE_API virtual FCompositeCorePassProxy* GetProxy(const UE::CompositeCore::FPassInputDecl& InputDecl, FCompositeTraversalContext& InContext, FSceneRenderingBulkObjectAllocator& InFrameAllocator) const override;
 
 	/** Get reflection & shadow caster actors. */
diff --git a/Engine/Plugins/Compositing/Composite/Source/Composite/Public/Layers/CompositeLayerShadowReflection.h b/Engine/Plugins/Compositing/Composite/Source/Composite/Public/Layers/CompositeLayerShadowReflection.h
--- a/Engine/Plugins/Compositing/Composite/Source/Composite/Public/Layers/CompositeLayerShadowReflection.h
+++ b/Engine/Plugins/Compositing/Composite/Source/Composite/Public/Layers/CompositeLayerShadowReflection.h
@@ -78,7 +78,7 @@ public:
 	/** End-of-frame update callback. */
 	UE_API virtual void OnEndOfFrameUpdate(UWorld* InWorld) override;
 
-	UE_API virtual bool IsActive() const override;
+	UE_API virtual bool GetIsActive() const override;
 	UE_API virtual FCompositeCorePassProxy* GetProxy(const UE::CompositeCore::FPassInputDecl& InputDecl, FCompositeTraversalContext& InContext, FSceneRenderingBulkObjectAllocator& InFrameAllocator) const override;
 
 	/** Get the render target resolution for the scene captures. */
@@ -98,7 +98,7 @@ public:
 	UE_API void SetActors(TArray<TSoftObjectPtr<AActor>> InActors);
 
 	/** Set the enabled state. */
-	UE_API virtual void SetEnabled(bool bInEnabled) override;
+	UE_API virtual void SetIsEnabled(bool bInEnabled) override;
 
 #if WITH_EDITORONLY_DATA
 PRAGMA_DISABLE_DEPRECATION_WARNINGS
diff --git a/Engine/Plugins/Compositing/Composite/Source/Composite/Public/Layers/CompositeLayerSingleLightShadow.h b/Engine/Plugins/Compositing/Composite/Source/Composite/Public/Layers/CompositeLayerSingleLightShadow.h
--- a/Engine/Plugins/Compositing/Composite/Source/Composite/Public/Layers/CompositeLayerSingleLightShadow.h
+++ b/Engine/Plugins/Compositing/Composite/Source/Composite/Public/Layers/CompositeLayerSingleLightShadow.h
@@ -54,7 +54,7 @@ public:
 	/** End-of-frame update callback. */
 	UE_API virtual void OnEndOfFrameUpdate(UWorld* InWorld) override;
 
-	UE_API virtual bool IsActive() const override;
+	UE_API virtual bool GetIsActive() const override;
 	UE_API virtual FCompositeCorePassProxy* GetProxy(const UE::CompositeCore::FPassInputDecl& InputDecl, FCompositeTraversalContext& InContext, FSceneRenderingBulkObjectAllocator& InFrameAllocator) const override;
 
 public:
diff --git a/Engine/Plugins/Compositing/Composite/Source/Composite/Public/Passes/CompositePassBase.h b/Engine/Plugins/Compositing/Composite/Source/Composite/Public/Passes/CompositePassBase.h
--- a/Engine/Plugins/Compositing/Composite/Source/Composite/Public/Passes/CompositePassBase.h
+++ b/Engine/Plugins/Compositing/Composite/Source/Composite/Public/Passes/CompositePassBase.h
@@ -62,19 +62,19 @@ public:
 	UE_API virtual bool CanEditChange(const FProperty* InProperty) const;
 #endif
 
-	/** Override to return a render-thread proxy for this pass. Proxy objects should be allocated from the provided allocator. Only called when IsActive() returns true. */
+	/** Override to return a render-thread proxy for this pass. Proxy objects should be allocated from the provided allocator. Only called when GetIsActive() returns true. */
 	UE_API virtual FCompositeCorePassProxy* GetProxy(const UE::CompositeCore::FPassInputDecl& InputDecl, FCompositeTraversalContext& InContext, FSceneRenderingBulkObjectAllocator& InFrameAllocator) const { return nullptr; }
 
 	/** Get the enabled state. */
 	UFUNCTION(BlueprintGetter)
-	UE_API virtual bool IsEnabled() const;
+	UE_API virtual bool GetIsEnabled() const;
 
 	/** Set the enabled state. */
-	UFUNCTION(BlueprintSetter)
-	UE_API virtual void SetEnabled(bool bInEnabled);
+	UFUNCTION(BlueprintSetter, CallInEditor)
+	UE_API virtual void SetIsEnabled(bool bInEnabled);
 
 	/** Returns true if the pass should participate in rendering this frame. Checked before GetProxy() is called. */
-	virtual bool IsActive() const { return IsEnabled(); }
+	virtual bool GetIsActive() const { return GetIsEnabled(); }
 
 #if WITH_EDITOR
 	/** Gets the display name of the pass */
@@ -109,7 +109,7 @@ public:
 	
 protected:
 	/** Whether or not the pass is active. */
-	UPROPERTY(EditAnywhere, BlueprintGetter = IsEnabled, BlueprintSetter = SetEnabled, Interp, Category = "Composite", meta = (DisplayPriority = "1"))
+	UPROPERTY(EditAnywhere, BlueprintGetter = GetIsEnabled, BlueprintSetter = SetIsEnabled, Interp, Category = "Composite", meta = (DisplayPriority = "1"))
 	bool bIsEnabled = true;
 
 #if WITH_EDITORONLY_DATA
diff --git a/Engine/Plugins/Compositing/Composite/Source/Composite/Public/Passes/CompositePassBlur.h b/Engine/Plugins/Compositing/Composite/Source/Composite/Public/Passes/CompositePassBlur.h
--- a/Engine/Plugins/Compositing/Composite/Source/Composite/Public/Passes/CompositePassBlur.h
+++ b/Engine/Plugins/Compositing/Composite/Source/Composite/Public/Passes/CompositePassBlur.h
@@ -28,7 +28,7 @@ public:
 	/** Destructor */
 	UE_API ~UCompositePassBlur();
 
-	UE_API virtual bool IsActive() const override;
+	UE_API virtual bool GetIsActive() const override;
 
 	UE_API virtual FCompositeCorePassProxy* GetProxy(const UE::CompositeCore::FPassInputDecl& InputDecl, FCompositeTraversalContext& InContext, FSceneRenderingBulkObjectAllocator& InFrameAllocator) const override;
 
diff --git a/Engine/Plugins/Compositing/Composite/Source/Composite/Public/Passes/CompositePassCenteredScale.h b/Engine/Plugins/Compositing/Composite/Source/Composite/Public/Passes/CompositePassCenteredScale.h
--- a/Engine/Plugins/Compositing/Composite/Source/Composite/Public/Passes/CompositePassCenteredScale.h
+++ b/Engine/Plugins/Compositing/Composite/Source/Composite/Public/Passes/CompositePassCenteredScale.h
@@ -57,7 +57,7 @@ public:
 
 	UE_API virtual FCompositeCorePassProxy* GetProxy(const UE::CompositeCore::FPassInputDecl& InputDecl, FCompositeTraversalContext& InContext, FSceneRenderingBulkObjectAllocator& InFrameAllocator) const override;
 
-	UE_API virtual bool IsActive() const override;
+	UE_API virtual bool GetIsActive() const override;
 
 	/** Calculate the final texture UV scale. */
 	UFUNCTION(BlueprintCallable, Category = "Composite")
@@ -97,7 +97,7 @@ public:
 	float Overscan;
 
 private:
-	/** Cached scale UV computed in IsActive() and reused in GetProxy(). Reset after consumption. */
+	/** Cached scale UV computed in GetIsActive() and reused in GetProxy(). Reset after consumption. */
 	mutable TOptional<FVector2f> CachedScaleUV;
 };
 
diff --git a/Engine/Plugins/Compositing/Composite/Source/Composite/Public/Passes/CompositePassDistortion.h b/Engine/Plugins/Compositing/Composite/Source/Composite/Public/Passes/CompositePassDistortion.h
--- a/Engine/Plugins/Compositing/Composite/Source/Composite/Public/Passes/CompositePassDistortion.h
+++ b/Engine/Plugins/Compositing/Composite/Source/Composite/Public/Passes/CompositePassDistortion.h
@@ -35,7 +35,7 @@ public:
 	virtual void PostInitProperties() override;
 	//~ End UObject Interface
 
-	virtual bool IsActive() const override;
+	virtual bool GetIsActive() const override;
 
 	UE_API virtual FCompositeCorePassProxy* GetProxy(const UE::CompositeCore::FPassInputDecl& InputDecl, FCompositeTraversalContext& InContext, FSceneRenderingBulkObjectAllocator& InFrameAllocator) const override;
 
diff --git a/Engine/Plugins/Compositing/Composite/Source/Composite/Public/Passes/CompositePassMasking.h b/Engine/Plugins/Compositing/Composite/Source/Composite/Public/Passes/CompositePassMasking.h
--- a/Engine/Plugins/Compositing/Composite/Source/Composite/Public/Passes/CompositePassMasking.h
+++ b/Engine/Plugins/Compositing/Composite/Source/Composite/Public/Passes/CompositePassMasking.h
@@ -31,7 +31,7 @@ public:
 	/** Destructor */
 	UE_API ~UCompositePassMasking();
 
-	UE_API virtual bool IsActive() const override;
+	UE_API virtual bool GetIsActive() const override;
 
 	UE_API virtual FCompositeCorePassProxy* GetProxy(const UE::CompositeCore::FPassInputDecl& InputDecl, FCompositeTraversalContext& InContext, FSceneRenderingBulkObjectAllocator& InFrameAllocator) const override;
 
diff --git a/Engine/Plugins/Compositing/Composite/Source/Composite/Public/Passes/CompositePassMaterial.h b/Engine/Plugins/Compositing/Composite/Source/Composite/Public/Passes/CompositePassMaterial.h
--- a/Engine/Plugins/Compositing/Composite/Source/Composite/Public/Passes/CompositePassMaterial.h
+++ b/Engine/Plugins/Compositing/Composite/Source/Composite/Public/Passes/CompositePassMaterial.h
@@ -23,7 +23,7 @@ public:
 	/** Destructor */
 	UE_API ~UCompositePassMaterial();
 
-	UE_API virtual bool IsActive() const override;
+	UE_API virtual bool GetIsActive() const override;
 
 	UE_API virtual FCompositeCorePassProxy* GetProxy(const UE::CompositeCore::FPassInputDecl& InputDecl, FCompositeTraversalContext& InContext, FSceneRenderingBulkObjectAllocator& InFrameAllocator) const override;
 
diff --git a/Engine/Plugins/Compositing/Composite/Source/CompositeEditor/Private/ActorFactories/CompositeActorFactory.cpp b/Engine/Plugins/Compositing/Composite/Source/CompositeEditor/Private/ActorFactories/CompositeActorFactory.cpp
--- a/Engine/Plugins/Compositing/Composite/Source/CompositeEditor/Private/ActorFactories/CompositeActorFactory.cpp
+++ b/Engine/Plugins/Compositing/Composite/Source/CompositeEditor/Private/ActorFactories/CompositeActorFactory.cpp
@@ -36,7 +36,7 @@ void UCompositeActorFactory::PostSpawnActor(UObject* Asset, AActor* NewActor)
 		NewLayers.Add(NewObject<UCompositeLayerPlate>(CompositeActor, TEXT("PlateLayer"), RF_Transactional));
 
 		// By default, we disable the shadow reflection catcher due to its significant added cost.
-		NewLayers[1]->SetEnabled(false);
+		NewLayers[1]->SetIsEnabled(false);
 
 		// No need to call Modify() on the spawned actor
 		CompositeActor->SetCompositeLayers(NewLayers);
diff --git a/Engine/Plugins/Compositing/Composite/Source/CompositeEditor/Private/CompositeEditorModule.cpp b/Engine/Plugins/Compositing/Composite/Source/CompositeEditor/Private/CompositeEditorModule.cpp
--- a/Engine/Plugins/Compositing/Composite/Source/CompositeEditor/Private/CompositeEditorModule.cpp
+++ b/Engine/Plugins/Compositing/Composite/Source/CompositeEditor/Private/CompositeEditorModule.cpp
@@ -135,7 +135,7 @@ void FCompositeEditorModule::OnLevelActorAdded(AActor* InActor)
 			ACompositeActor* CompositeActor = Cast<ACompositeActor>(InActor);
 			if (IsValid(CompositeActor))
 			{
-				CompositeActor->SetActive(false);
+				CompositeActor->SetIsActive(false);
 
 				UE_LOG(LogCompositeEditor, Display, TEXT("No more than one composite actor should be active at the same time."));
 			}
diff --git a/Engine/Plugins/Compositing/Composite/Source/CompositeEditor/Private/Sequencer/CompositeDetailKeyframeHandler.cpp b/Engine/Plugins/Compositing/Composite/Source/CompositeEditor/Private/Sequencer/CompositeDetailKeyframeHandler.cpp
new file mode 100644
--- /dev/null
+++ b/Engine/Plugins/Compositing/Composite/Source/CompositeEditor/Private/Sequencer/CompositeDetailKeyframeHandler.cpp
@@ -0,0 +1,129 @@
+// Copyright Epic Games, Inc. All Rights Reserved.
+
+#include "CompositeDetailKeyframeHandler.h"
+
+#include "PropertyHandle.h"
+#include "ISequencer.h"
+#include "LevelEditorSequencerIntegration.h"
+#include "MovieScene.h"
+#include "PropertyPath.h"
+
+TArray<TWeakPtr<ISequencer>> FCompositeDetailKeyframeHandler::GetSequencers() const
+{
+	return FLevelEditorSequencerIntegration::Get().GetSequencers();
+}
+
+bool FCompositeDetailKeyframeHandler::IsPropertyKeyable(const UClass* InObjectClass, const IPropertyHandle& PropertyHandle) const
+{
+	FCanKeyPropertyParams CanKeyPropertyParams(InObjectClass, PropertyHandle);
+
+	for (const TWeakPtr<ISequencer>& WeakSequencer : GetSequencers())
+	{
+		TSharedPtr<ISequencer> Sequencer = WeakSequencer.Pin();
+		if (Sequencer.IsValid() && Sequencer->CanKeyProperty(CanKeyPropertyParams) && !Sequencer->IsReadOnly())
+		{
+			return true;
+		}
+	}
+
+	return false;
+}
+
+bool FCompositeDetailKeyframeHandler::IsPropertyKeyingEnabled() const
+{
+	for (const TWeakPtr<ISequencer>& WeakSequencer : GetSequencers())
+	{
+		TSharedPtr<ISequencer> Sequencer = WeakSequencer.Pin();
+		if (Sequencer.IsValid() && Sequencer->GetFocusedMovieSceneSequence() && Sequencer->GetAllowEditsMode() != EAllowEditsMode::AllowLevelEditsOnly)
+		{
+			return true;
+		}
+	}
+	
+	return false;
+}
+
+bool FCompositeDetailKeyframeHandler::IsPropertyAnimated(const IPropertyHandle& PropertyHandle, UObject* ParentObject) const
+{
+	for (const TWeakPtr<ISequencer>& WeakSequencer : GetSequencers())
+	{
+		TSharedPtr<ISequencer> Sequencer = WeakSequencer.Pin();
+		if (Sequencer.IsValid() && Sequencer->GetFocusedMovieSceneSequence())
+		{
+			constexpr bool bCreateHandleIfMissing = false;
+			FGuid ObjectHandle = Sequencer->GetHandleToObject(ParentObject, bCreateHandleIfMissing);
+			if (ObjectHandle.IsValid())
+			{
+				UMovieScene* MovieScene = Sequencer->GetFocusedMovieSceneSequence()->GetMovieScene();
+
+				/**
+				 * Build the dot-separated property path that Sequencer uses for track lookup.
+				 * Walk parent handles to produce the full path, e.g. "StructName.MemberName" for nested properties.
+				 */
+				FProperty* LeafProperty = PropertyHandle.GetProperty();
+				if (!LeafProperty)
+				{
+					continue;
+				}
+
+				TArray<FProperty*> PropertyChain;
+				PropertyChain.Add(LeafProperty);
+				for (TSharedPtr<IPropertyHandle> Current = PropertyHandle.GetParentHandle(); Current.IsValid() && Current->GetProperty(); Current = Current->GetParentHandle())
+				{
+					PropertyChain.Add(Current->GetProperty());
+				}
+
+				TSharedRef<FPropertyPath> PropertyPath = FPropertyPath::CreateEmpty();
+				for (int32 Index = PropertyChain.Num() - 1; Index >= 0; --Index)
+				{
+					PropertyPath->AddProperty(FPropertyInfo(PropertyChain[Index]));
+				}
+				FName PropertyName(*PropertyPath->ToString(TEXT(".")));
+
+				if (MovieScene->FindTrack(UMovieSceneTrack::StaticClass(), ObjectHandle, PropertyName) != nullptr)
+				{
+					return true;
+				}
+			}
+		}
+	}
+	
+	return false;
+}
+
+void FCompositeDetailKeyframeHandler::OnKeyPropertyClicked(const IPropertyHandle& KeyedPropertyHandle)
+{
+	TArray<UObject*> Objects;
+	KeyedPropertyHandle.GetOuterObjects(Objects);
+
+	TArray<UObject*> EachObject;
+	EachObject.SetNum(1);
+	for (const TWeakPtr<ISequencer>& WeakSequencer : GetSequencers())
+	{
+		TSharedPtr<ISequencer> Sequencer = WeakSequencer.Pin();
+		if (Sequencer.IsValid())
+		{
+			for (UObject* Object : Objects)
+			{
+				EachObject[0] = Object;
+				FKeyPropertyParams KeyPropertyParams(EachObject, KeyedPropertyHandle, ESequencerKeyMode::ManualKeyForced);
+				Sequencer->KeyProperty(KeyPropertyParams);
+			}
+		}
+	}
+}
+
+EPropertyKeyedStatus FCompositeDetailKeyframeHandler::GetPropertyKeyedStatus(const IPropertyHandle& PropertyHandle) const
+{
+	EPropertyKeyedStatus KeyedStatus = EPropertyKeyedStatus::NotKeyed;
+	for (const TWeakPtr<ISequencer>& WeakSequencer : GetSequencers())
+	{
+		if (TSharedPtr<ISequencer> Sequencer = WeakSequencer.Pin())
+		{
+			EPropertyKeyedStatus NewKeyedStatus = Sequencer->GetPropertyKeyedStatus(PropertyHandle);
+			KeyedStatus = FMath::Max(KeyedStatus, NewKeyedStatus);
+		}
+	}
+	
+	return KeyedStatus;
+}
diff --git a/Engine/Plugins/Compositing/Composite/Source/CompositeEditor/Private/Sequencer/CompositeDetailKeyframeHandler.h b/Engine/Plugins/Compositing/Composite/Source/CompositeEditor/Private/Sequencer/CompositeDetailKeyframeHandler.h
new file mode 100644
--- /dev/null
+++ b/Engine/Plugins/Compositing/Composite/Source/CompositeEditor/Private/Sequencer/CompositeDetailKeyframeHandler.h
@@ -0,0 +1,25 @@
+// Copyright Epic Games, Inc. All Rights Reserved.
+
+#pragma once
+
+#include "IDetailKeyframeHandler.h"
+
+class ISequencer;
+
+/**
+ * Keyframe handler for the Composite editor panel's DetailsView.
+ * Delegates to active Sequencer instances obtained from FLevelEditorSequencerIntegration.
+ * Stateless — queries live sequencer state on each call so it handles Sequencer open/close automatically.
+ */
+class FCompositeDetailKeyframeHandler : public IDetailKeyframeHandler
+{
+public:
+	virtual bool IsPropertyKeyable(const UClass* InObjectClass, const IPropertyHandle& PropertyHandle) const override;
+	virtual bool IsPropertyKeyingEnabled() const override;
+	virtual bool IsPropertyAnimated(const IPropertyHandle& PropertyHandle, UObject* ParentObject) const override;
+	virtual void OnKeyPropertyClicked(const IPropertyHandle& KeyedPropertyHandle) override;
+	virtual EPropertyKeyedStatus GetPropertyKeyedStatus(const IPropertyHandle& PropertyHandle) const override;
+
+private:
+	TArray<TWeakPtr<ISequencer>> GetSequencers() const;
+};
diff --git a/Engine/Plugins/Compositing/Composite/Source/CompositeEditor/Private/UI/SCompositeEditorPanel.cpp b/Engine/Plugins/Compositing/Composite/Source/CompositeEditor/Private/UI/SCompositeEditorPanel.cpp
--- a/Engine/Plugins/Compositing/Composite/Source/CompositeEditor/Private/UI/SCompositeEditorPanel.cpp
+++ b/Engine/Plugins/Compositing/Composite/Source/CompositeEditor/Private/UI/SCompositeEditorPanel.cpp
@@ -18,6 +18,7 @@
 #include "WorkspaceMenuStructureModule.h"
 #include "Customizations/CompositeLayerPlateCustomization.h"
 #include "Customizations/CompositeLayerSceneCaptureCustomization.h"
+#include "Sequencer/CompositeDetailKeyframeHandler.h"
 #include "Layers/CompositeLayerPlate.h"
 #include "Layers/CompositeLayerSceneCapture.h"
 #include "Widgets/Docking/SDockTab.h"
@@ -186,6 +187,7 @@ void SCompositeEditorPanel::Construct(const FArguments& InArgs)
 	DetailsViewArgs.bShowPropertyMatrixButton = false;
 	
 	DetailsView = FModuleManager::GetModuleChecked<FPropertyEditorModule>("PropertyEditor").CreateDetailView(DetailsViewArgs);
+	DetailsView->SetKeyframeHandler(MakeShared<FCompositeDetailKeyframeHandler>());
 	DetailsView->RegisterInstancedCustomPropertyLayout(ACompositeActor::StaticClass(), FOnGetDetailCustomizationInstance::CreateLambda([]
 	{
 		return MakeShared<FCompositeActorPanelDetailCustomization>();
diff --git a/Engine/Plugins/Compositing/Composite/Source/CompositeEditor/Private/UI/SCompositePanelLayerTree.cpp b/Engine/Plugins/Compositing/Composite/Source/CompositeEditor/Private/UI/SCompositePanelLayerTree.cpp
--- a/Engine/Plugins/Compositing/Composite/Source/CompositeEditor/Private/UI/SCompositePanelLayerTree.cpp
+++ b/Engine/Plugins/Compositing/Composite/Source/CompositeEditor/Private/UI/SCompositePanelLayerTree.cpp
@@ -504,7 +504,7 @@ private:
 	{
 		if (TStrongObjectPtr<ACompositeActor> CompositeActor = GetCompositeActor())
 		{
-			return CompositeActor->IsActive() ? ECheckBoxState::Checked : ECheckBoxState::Unchecked;
+			return CompositeActor->GetIsActive() ? ECheckBoxState::Checked : ECheckBoxState::Unchecked;
 		}
 
 		return ECheckBoxState::Unchecked;
@@ -528,7 +528,7 @@ private:
 		{
 			const bool bIsActive = (InCheckBoxState == ECheckBoxState::Checked);
 			// No transactions, the property is non-transactional
-			CompositeActor->SetActive(bIsActive);
+			CompositeActor->SetIsActive(bIsActive);
 			OnCompositeActorActivated.ExecuteIfBound(bIsActive ? CompositeActor : nullptr);
 		}
 	}
@@ -540,7 +540,7 @@ private:
 		{
 			if (UCompositeLayerBase* CompositeLayer = TreeItem->GetCompositeLayer())
 			{
-				return CompositeLayer->IsEnabled() ? ECheckBoxState::Checked : ECheckBoxState::Unchecked;
+				return CompositeLayer->GetIsEnabled() ? ECheckBoxState::Checked : ECheckBoxState::Unchecked;
 			}
 			
 			// In the case where there is a valid layer index, but the layer itself is null
@@ -549,7 +549,7 @@ private:
 
 		if (TStrongObjectPtr<ACompositeActor> CompositeActor = GetCompositeActor())
 		{
-			return CompositeActor->IsEnabled() ? ECheckBoxState::Checked : ECheckBoxState::Unchecked;
+			return CompositeActor->GetIsEnabled() ? ECheckBoxState::Checked : ECheckBoxState::Unchecked;
 		}
 		
 		return ECheckBoxState::Unchecked;
@@ -565,20 +565,20 @@ private:
 				FScopedTransaction SetEnabledTransaction(LOCTEXT("SetEnabledTransaction", "Set Enabled"));
 				CompositeLayer->Modify();
 
-				CompositeLayer->SetEnabled(InCheckBoxState == ECheckBoxState::Checked);
+				CompositeLayer->SetIsEnabled(InCheckBoxState == ECheckBoxState::Checked);
 				return;
 			}
 
 			// In the case where there is a valid layer index, but the layer itself is null
 			return;
 		}
-		
+
 		if (TStrongObjectPtr<ACompositeActor> CompositeActor = GetCompositeActor())
 		{
 			FScopedTransaction SetEnabledTransaction(LOCTEXT("SetEnabledTransaction", "Set Enabled"));
 			CompositeActor->Modify();
 
-			CompositeActor->SetEnabled(InCheckBoxState == ECheckBoxState::Checked);
+			CompositeActor->SetIsEnabled(InCheckBoxState == ECheckBoxState::Checked);
 			return;
 		}
 	}
@@ -1503,7 +1503,7 @@ bool SCompositePanelLayerTree::CanPilotCamera() const
 		// If the active composite actor is not included in the currently selected items, disable the Pilot button
 		const bool bSelectionContainsActiveCompositeActor = SelectedItems.ContainsByPredicate([](const FCompositeActorTreeItemPtr& InTreeItem)
 		{
-			return InTreeItem.IsValid() && InTreeItem->HasValidCompositeActor() && InTreeItem->CompositeActor->IsActive();
+			return InTreeItem.IsValid() && InTreeItem->HasValidCompositeActor() && InTreeItem->CompositeActor->GetIsActive();
 		});
 
 		if (!bSelectionContainsActiveCompositeActor)
@@ -1531,7 +1531,7 @@ void SCompositePanelLayerTree::OnPilotCameraToggled(bool bPilotCamera)
 			// If the active composite actor is not included in the currently selected items, don't try and pilot anything
 			const bool bSelectionContainsActiveCompositeActor = SelectedItems.ContainsByPredicate([](const FCompositeActorTreeItemPtr& InTreeItem)
 			{
-				return InTreeItem.IsValid() && InTreeItem->HasValidCompositeActor() && InTreeItem->CompositeActor->IsActive();
+				return InTreeItem.IsValid() && InTreeItem->HasValidCompositeActor() && InTreeItem->CompositeActor->GetIsActive();
 			});
 
 			if (!bSelectionContainsActiveCompositeActor)
@@ -1560,7 +1560,7 @@ AActor* SCompositePanelLayerTree::GetActiveCompositeActorCamera() const
 	TStrongObjectPtr<ACompositeActor> ActiveCompositeActor = nullptr;
 	const FCompositeActorTreeItemPtr* ActiveTreeItem = CompositeActorTreeItems.FindByPredicate([](const FCompositeActorTreeItemPtr& InTreeItem)
 	{
-		return InTreeItem.IsValid() && InTreeItem->HasValidCompositeActor() && InTreeItem->CompositeActor->IsActive();
+		return InTreeItem.IsValid() && InTreeItem->HasValidCompositeActor() && InTreeItem->CompositeActor->GetIsActive();
 	});
 	
 	if (ActiveTreeItem)
@@ -1672,7 +1672,7 @@ ECheckBoxState SCompositePanelLayerTree::GetGlobalActiveCheckState() const
 			continue;
 		}
 
-		if (TreeItem->CompositeActor->IsActive())
+		if (TreeItem->CompositeActor->GetIsActive())
 		{
 			bAnyEnabled = true;
 		}
@@ -1698,7 +1698,7 @@ void SCompositePanelLayerTree::OnGlobalActiveCheckStateChanged(ECheckBoxState Ch
 			{
 				if (CompositeActor->IsInVPRole())
 				{
-					CompositeActor->SetActive(true);
+					CompositeActor->SetIsActive(true);
 					OnCompositeActorActivated(CompositeActor);
 					bActivated = true;
 				}
@@ -1718,7 +1718,7 @@ void SCompositePanelLayerTree::OnGlobalActiveCheckStateChanged(ECheckBoxState Ch
 				TStrongObjectPtr<ACompositeActor> CompositeActor = TreeItem->CompositeActor.Pin();
 				if (CompositeActor->IsInVPRole())
 				{
-					CompositeActor->SetActive(true);
+					CompositeActor->SetIsActive(true);
 					OnCompositeActorActivated(CompositeActor);
 					break;
 				}
@@ -1734,12 +1734,12 @@ void SCompositePanelLayerTree::OnGlobalActiveCheckStateChanged(ECheckBoxState Ch
 				continue;
 			}
 
-			if (TreeItem->CompositeActor->IsActive())
+			if (TreeItem->CompositeActor->GetIsActive())
 			{
 				LastActiveActor = TreeItem->CompositeActor;
 			}
 
-			TreeItem->CompositeActor->SetActive(false);
+			TreeItem->CompositeActor->SetIsActive(false);
 			OnCompositeActorActivated(nullptr);
 		}
 	}
@@ -1756,7 +1756,7 @@ ECheckBoxState SCompositePanelLayerTree::GetGlobalEnabledCheckState() const
 			continue;
 		}
 
-		if (TreeItem->CompositeActor->IsEnabled())
+		if (TreeItem->CompositeActor->GetIsEnabled())
 		{
 			bAnyEnabled = true;
 		}
@@ -1776,7 +1776,7 @@ ECheckBoxState SCompositePanelLayerTree::GetGlobalEnabledCheckState() const
 
 			if (UCompositeLayerBase* CompositeLayer = ChildTreeItem->GetCompositeLayer())
 			{
-				if (CompositeLayer->IsEnabled())
+				if (CompositeLayer->GetIsEnabled())
 				{
 					bAnyEnabled = true;
 				}
@@ -1815,7 +1815,7 @@ void SCompositePanelLayerTree::OnGlobalEnabledCheckStateChanged(ECheckBoxState C
 		}
 
 		TreeItem->CompositeActor->Modify();
-		TreeItem->CompositeActor->SetEnabled(CheckBoxState == ECheckBoxState::Checked);
+		TreeItem->CompositeActor->SetIsEnabled(CheckBoxState == ECheckBoxState::Checked);
 		
 		for (const FCompositeActorTreeItemPtr& ChildTreeItem : TreeItem->Children)
 		{
@@ -1829,7 +1829,7 @@ void SCompositePanelLayerTree::OnGlobalEnabledCheckStateChanged(ECheckBoxState C
 			if (UCompositeLayerBase* CompositeLayer = ChildTreeItem->GetCompositeLayer())
 			{
 				CompositeLayer->Modify();
-				CompositeLayer->SetEnabled(CheckBoxState == ECheckBoxState::Checked);
+				CompositeLayer->SetIsEnabled(CheckBoxState == ECheckBoxState::Checked);
 			}
 		}
 	}
@@ -2684,7 +2684,7 @@ void SCompositePanelLayerTree::EnableSelectedItems()
 		{
 			if (UCompositeLayerBase* Layer = Item->GetCompositeLayer())
 			{
-				if (Layer->IsEnabled())
+				if (Layer->GetIsEnabled())
 				{
 					bAnyEnabled = true;
 				}
@@ -2692,7 +2692,7 @@ void SCompositePanelLayerTree::EnableSelectedItems()
 		}
 		else
 		{
-			if (PinnedCompositeActor->IsEnabled())
+			if (PinnedCompositeActor->GetIsEnabled())
 			{
 				bAnyEnabled = true;
 			}
@@ -2715,13 +2715,13 @@ void SCompositePanelLayerTree::EnableSelectedItems()
 			if (UCompositeLayerBase* Layer = Item->GetCompositeLayer())
 			{
 				Layer->Modify();
-				Layer->SetEnabled(bSetEnabled);
+				Layer->SetIsEnabled(bSetEnabled);
 			}
 		}
 		else
 		{
 			PinnedCompositeActor->Modify();
-			PinnedCompositeActor->SetEnabled(bSetEnabled);
+			PinnedCompositeActor->SetIsEnabled(bSetEnabled);
 		}
 	}
 }
diff --git a/Engine/Plugins/Compositing/Composite/Source/CompositeEditor/Private/UI/SCompositePassTree.cpp b/Engine/Plugins/Compositing/Composite/Source/CompositeEditor/Private/UI/SCompositePassTree.cpp
--- a/Engine/Plugins/Compositing/Composite/Source/CompositeEditor/Private/UI/SCompositePassTree.cpp
+++ b/Engine/Plugins/Compositing/Composite/Source/CompositeEditor/Private/UI/SCompositePassTree.cpp
@@ -648,9 +648,9 @@ private:
 		{
 			if (UCompositePassBase* Pass = TreeItem->GetPass())
 			{
-				return Pass->IsEnabled() ? ECheckBoxState::Checked : ECheckBoxState::Unchecked;
+				return Pass->GetIsEnabled() ? ECheckBoxState::Checked : ECheckBoxState::Unchecked;
 			}
-			
+
 			// In the case where there is a valid pass index, but the pass itself is null
 			return ECheckBoxState::Unchecked;
 		}
@@ -668,7 +668,7 @@ private:
 
 			if (UCompositePassBase* Pass = ChildTreeItem->GetPass())
 			{
-				if (Pass->IsEnabled())
+				if (Pass->GetIsEnabled())
 				{
 					bAnyEnabled = true;
 				}
@@ -720,7 +720,7 @@ private:
 				FScopedTransaction SetEnabledTransaction(LOCTEXT("SetEnabledTransaction", "Set Enabled"));
 				Pass->Modify();
 
-				Pass->SetEnabled(InCheckBoxState == ECheckBoxState::Checked);
+				Pass->SetIsEnabled(InCheckBoxState == ECheckBoxState::Checked);
 				return;
 			}
 
@@ -740,7 +740,7 @@ private:
 			if (UCompositePassBase* Pass = ChildTreeItem->GetPass())
 			{
 				Pass->Modify();
-				Pass->SetEnabled(InCheckBoxState == ECheckBoxState::Checked);
+				Pass->SetIsEnabled(InCheckBoxState == ECheckBoxState::Checked);
 			}
 		}
 	}
@@ -1480,7 +1480,7 @@ ECheckBoxState SCompositePassTree::GetGlobalEnabledState() const
 				continue;
 			}
 			
-			if (Pass->IsEnabled())
+			if (Pass->GetIsEnabled())
 			{
 				bAnyPassesEnabled = true;
 			}
@@ -1569,7 +1569,7 @@ void SCompositePassTree::OnGlobalEnabledStateChanged(ECheckBoxState CheckBoxStat
 			}
 
 			Pass->Modify();
-			Pass->SetEnabled(CheckBoxState == ECheckBoxState::Checked);
+			Pass->SetIsEnabled(CheckBoxState == ECheckBoxState::Checked);
 		}
 	};
 	
@@ -2220,7 +2220,7 @@ void SCompositePassTree::EnableSelectedItems()
 		{
 			if (UCompositePassBase* Pass = Item->GetPass())
 			{
-				if (Pass->IsEnabled())
+				if (Pass->GetIsEnabled())
 				{
 					bAnyEnabled = true;
 				}
@@ -2234,7 +2234,7 @@ void SCompositePassTree::EnableSelectedItems()
 				{
 					if (UCompositePassBase* Pass = ChildItem->GetPass())
 					{
-						if (Pass->IsEnabled())
+						if (Pass->GetIsEnabled())
 						{
 							bAnyEnabled = true;
 						}
@@ -2256,7 +2256,7 @@ void SCompositePassTree::EnableSelectedItems()
 			if (UCompositePassBase* Pass = Item->GetPass())
 			{
 				Pass->Modify();
-				Pass->SetEnabled(bSetEnabled);
+				Pass->SetIsEnabled(bSetEnabled);
 			}
 		}
 		else
@@ -2268,7 +2268,7 @@ void SCompositePassTree::EnableSelectedItems()
 					if (UCompositePassBase* Pass = ChildItem->GetPass())
 					{
 						Pass->Modify();
-						Pass->SetEnabled(bSetEnabled);
+						Pass->SetIsEnabled(bSetEnabled);
 					}
 				}
 			}
