Composure: Add support for inverted alpha handling to Blur, ColorGrading, and OpenColorIO passes.

* Specifically needed for passes on the main render layer (recent addition).
* Removed undesirable unpremult property on the blur pass , we prefer always blurring with alpha premultiplied. Previous behavior was present because inverted alpha wasn't handled correctly.

#Synced-CL 52383054

diff --git a/Engine/Plugins/Compositing/Composite/Shaders/Private/CompositeBlur.usf b/Engine/Plugins/Compositing/Composite/Shaders/Private/CompositeBlur.usf
--- a/Engine/Plugins/Compositing/Composite/Shaders/Private/CompositeBlur.usf
+++ b/Engine/Plugins/Compositing/Composite/Shaders/Private/CompositeBlur.usf
@@ -13,7 +13,7 @@ FScreenTransform SvPositionToInputTextureUV;
 SCREEN_PASS_TEXTURE_VIEWPORT(Input)
 float2 Dir;
 float Radius;
-uint bUnpremult;
+uint bInvertedAlpha;
 
 #if ALPHA_ONLY
     #define BLUR_TYPE float
@@ -34,9 +34,13 @@ float2 ClampInputUV(float2 UV)
     return clamp(UV, Input_UVViewportBilinearMin, Input_UVViewportBilinearMax);
 }
 
+/** Samples the input texture, normalizing inverted alpha (e.g. UE scene color) to standard opacity for processing. */
 float4 SampleInputTexture(float2 UV)
 {
-    return Input_Texture.SampleLevel(Input_Sampler, ClampInputUV(UV), 0.0);
+    float4 Color = Input_Texture.SampleLevel(Input_Sampler, ClampInputUV(UV), 0.0);
+    Color.a = InvertAlpha(Color.a, bInvertedAlpha);
+    
+    return Color;
 }
 
 BLUR_TYPE GaussianBlur(float2 UV)
@@ -98,17 +102,10 @@ void MainPS(
 
     BLUR_TYPE Blurred = GaussianBlur(InputUV);
 
+/** Restore original alpha convention on output (symmetric with SampleInputTexture). */
 #if ALPHA_ONLY
-    float4 Original = SampleInputTexture(InputUV);
-    if (bUnpremult)
-    {
-        OutColor = float4(Original.rgb * SafeRcp(Original.a) * Blurred, Blurred);
-    }
-    else
-    {
-        OutColor = float4(Original.rgb, Blurred);
-    }
+    OutColor = float4(SampleInputTexture(InputUV).rgb, InvertAlpha(Blurred, bInvertedAlpha));
 #else
-    OutColor = Blurred;
+    OutColor = float4(Blurred.rgb, InvertAlpha(Blurred.a, bInvertedAlpha));
 #endif
 }
diff --git a/Engine/Plugins/Compositing/Composite/Shaders/Private/CompositeColorGrading.usf b/Engine/Plugins/Compositing/Composite/Shaders/Private/CompositeColorGrading.usf
--- a/Engine/Plugins/Compositing/Composite/Shaders/Private/CompositeColorGrading.usf
+++ b/Engine/Plugins/Compositing/Composite/Shaders/Private/CompositeColorGrading.usf
@@ -44,6 +44,7 @@ float4 Highlights_Gain;
 float4 Highlights_Offset;
 
 uint bUnpremult;
+uint bInvertedAlpha;
 float ShadowsMax;
 float HighlightsMin;
 float HighlightsMax;
@@ -109,11 +110,14 @@ void MainPS(
     float2 InputUV = ApplyScreenTransform(SvPosition.xy, SvPositionToInputTextureUV);
     float4 LinearColor = Input_Texture.Sample(Input_Sampler, InputUV);
 	
+	/** Normalize inverted alpha (e.g. UE scene color) to standard opacity for processing. */
+	LinearColor.a = InvertAlpha(LinearColor.a, bInvertedAlpha);
+
 	if(bUnpremult)
 	{
 		LinearColor = Unpremultiply(LinearColor);
 	}
-    
+
     BRANCH
 	if( abs(WhiteTemp - DefaultTemperature) > Epsilon || abs(WhiteTint - DefaultTint) > Epsilon )
 	{
@@ -131,5 +135,8 @@ void MainPS(
 		LinearColor = Premultiply(LinearColor);
 	}
 
+	/** Restore original alpha convention on output. */
+	LinearColor.a = InvertAlpha(LinearColor.a, bInvertedAlpha);
+
 	OutColor = LinearColor;
 }
diff --git a/Engine/Plugins/Compositing/Composite/Source/Composite/Private/Passes/CompositePassBlur.cpp b/Engine/Plugins/Compositing/Composite/Source/Composite/Private/Passes/CompositePassBlur.cpp
--- a/Engine/Plugins/Compositing/Composite/Source/Composite/Private/Passes/CompositePassBlur.cpp
+++ b/Engine/Plugins/Compositing/Composite/Source/Composite/Private/Passes/CompositePassBlur.cpp
@@ -23,7 +23,7 @@ class FCompositePassBlurShader : public FGlobalShader
 		SHADER_PARAMETER(FScreenTransform, SvPositionToInputTextureUV)
 		SHADER_PARAMETER(FVector2f, Dir)
 		SHADER_PARAMETER(float, Radius)
-		SHADER_PARAMETER(uint32, bUnpremult)
+		SHADER_PARAMETER(uint32, bInvertedAlpha)
 		RENDER_TARGET_BINDING_SLOTS()
 	END_SHADER_PARAMETER_STRUCT()
 };
@@ -72,7 +72,7 @@ namespace UE
 						FScreenTransform::ChangeTextureBasisFromTo(FScreenPassTextureViewport(InputTexture), FScreenTransform::ETextureBasis::ViewportUV, FScreenTransform::ETextureBasis::TextureUV));
 					Parameters->Dir = FVector2f(0, 1);
 					Parameters->Radius = Radius;
-					Parameters->bUnpremult = static_cast<uint32>(bUnpremult);
+					Parameters->bInvertedAlpha = static_cast<uint32>(Metadata.bInvertedAlpha);
 					Parameters->RenderTargets[0] = IntermediateOutput.GetRenderTargetBinding();
 
 					TShaderMapRef<FCompositePassBlurShader> PixelShader(GlobalShaderMap, PermutationVector);
@@ -94,7 +94,7 @@ namespace UE
 						FScreenTransform::ChangeTextureBasisFromTo(FScreenPassTextureViewport(IntermediateOutput), FScreenTransform::ETextureBasis::ViewportUV, FScreenTransform::ETextureBasis::TextureUV));
 					Parameters->Dir = FVector2f(1, 0);
 					Parameters->Radius = Radius;
-					Parameters->bUnpremult = static_cast<uint32>(bUnpremult);
+					Parameters->bInvertedAlpha = static_cast<uint32>(Metadata.bInvertedAlpha);
 					Parameters->RenderTargets[0] = Output.GetRenderTargetBinding();
 
 					TShaderMapRef<FCompositePassBlurShader> PixelShader(GlobalShaderMap, PermutationVector);
@@ -113,7 +113,6 @@ namespace UE
 
 			int32 Radius = 1;
 			bool bAlphaOnly = false;
-			bool bUnpremult = true;
 		};
 	}
 }
@@ -123,7 +122,6 @@ UCompositePassBlur::UCompositePassBlur(const FObjectInitializer& ObjectInitializ
 	, Method(ECompositePassBlurMethod::Gaussian)
 	, Radius(1)
 	, bAlphaOnly(false)
-	, bUnpremult(true)
 {
 }
 
@@ -141,7 +139,6 @@ FCompositeCorePassProxy* UCompositePassBlur::GetProxy(const UE::CompositeCore::F
 	FCompositePassBlurProxy* Proxy = InFrameAllocator.Create<FCompositePassBlurProxy>(UE::CompositeCore::FPassInputDeclArray{ InputDecl });
 	Proxy->Radius = Radius;
 	Proxy->bAlphaOnly = bAlphaOnly;
-	Proxy->bUnpremult = bUnpremult;
 
 	return Proxy;
 }
diff --git a/Engine/Plugins/Compositing/Composite/Source/Composite/Private/Passes/CompositePassColorGrading.cpp b/Engine/Plugins/Compositing/Composite/Source/Composite/Private/Passes/CompositePassColorGrading.cpp
--- a/Engine/Plugins/Compositing/Composite/Source/Composite/Private/Passes/CompositePassColorGrading.cpp
+++ b/Engine/Plugins/Compositing/Composite/Source/Composite/Private/Passes/CompositePassColorGrading.cpp
@@ -32,6 +32,7 @@ class FCompositeColorGradingShader : public FGlobalShader
 		SHADER_PARAMETER_STRUCT(FColorGradingPerRangeParameters, Midtones)
 		SHADER_PARAMETER_STRUCT(FColorGradingPerRangeParameters, Highlights)
 		SHADER_PARAMETER(uint32, bUnpremult)
+		SHADER_PARAMETER(uint32, bInvertedAlpha)
 		SHADER_PARAMETER(float, ShadowsMax)
 		SHADER_PARAMETER(float, HighlightsMin)
 		SHADER_PARAMETER(float, HighlightsMax)
@@ -89,6 +90,7 @@ namespace UE::Composite::Private
 		Parameters->Midtones = FCompositeColorGradingShader::ToParameters(ColorGradingSettings.Midtones);
 		Parameters->Highlights = FCompositeColorGradingShader::ToParameters(ColorGradingSettings.Highlights);
 		Parameters->bUnpremult = bUnpremult;
+		Parameters->bInvertedAlpha = static_cast<uint32>(Metadata.bInvertedAlpha);
 		Parameters->ShadowsMax = ColorGradingSettings.ShadowsMax;
 		Parameters->HighlightsMin = ColorGradingSettings.HighlightsMin;
 		Parameters->HighlightsMax = ColorGradingSettings.HighlightsMax;
diff --git a/Engine/Plugins/Compositing/Composite/Source/Composite/Private/Passes/CompositePassOpenColorIO.cpp b/Engine/Plugins/Compositing/Composite/Source/Composite/Private/Passes/CompositePassOpenColorIO.cpp
--- a/Engine/Plugins/Compositing/Composite/Source/Composite/Private/Passes/CompositePassOpenColorIO.cpp
+++ b/Engine/Plugins/Compositing/Composite/Source/Composite/Private/Passes/CompositePassOpenColorIO.cpp
@@ -38,22 +38,34 @@ public:
 			Output = CreateOutputRenderTarget(GraphBuilder, InView, PassContext.OutputViewRect, Input.Texture->Desc, TEXT("OcioCompositePass"));
 		}
 
+		EOpenColorIOTransformAlpha TransformAlpha = EOpenColorIOTransformAlpha::None;
+		if (bUnpremult)
+		{
+			TransformAlpha = Metadata.bInvertedAlpha
+				? EOpenColorIOTransformAlpha::InvertUnpremultiply
+				: EOpenColorIOTransformAlpha::Unpremultiply;
+		}
+
 		FOpenColorIORendering::AddPass_RenderThread(
 			GraphBuilder,
 			InView,
 			Input,
 			Output,
-			PassResources
+			PassResources,
+			TransformAlpha
 		);
 
 		return UE::CompositeCore::FPassTexture{ MoveTemp(Output), Metadata };
 	}
 
 	FOpenColorIORenderPassResources PassResources;
+	bool bUnpremult = true;
 };
 
 UCompositePassOpenColorIO::UCompositePassOpenColorIO(const FObjectInitializer& ObjectInitializer)
 	: Super(ObjectInitializer)
+	, ColorConversionSettings{}
+	, bUnpremult{true}
 {
 }
 
@@ -61,8 +73,11 @@ UCompositePassOpenColorIO::~UCompositePassOpenColorIO() = default;
 
 FCompositeCorePassProxy* UCompositePassOpenColorIO::GetProxy(const UE::CompositeCore::FPassInputDecl& InputDecl, FCompositeTraversalContext& InContext, FSceneRenderingBulkObjectAllocator& InFrameAllocator) const
 {
-	return InFrameAllocator.Create<FOpenColorIOCompositePassProxy>(UE::CompositeCore::FPassInputDeclArray{ InputDecl },
+	FOpenColorIOCompositePassProxy* Proxy = InFrameAllocator.Create<FOpenColorIOCompositePassProxy>(UE::CompositeCore::FPassInputDeclArray{ InputDecl },
 		FOpenColorIORendering::GetRenderPassResources(ColorConversionSettings, GMaxRHIFeatureLevel)
 	);
+	Proxy->bUnpremult = bUnpremult;
+
+	return Proxy;
 }
 
diff --git a/Engine/Plugins/Compositing/Composite/Source/Composite/Public/Passes/CompositePassBlur.h b/Engine/Plugins/Compositing/Composite/Source/Composite/Public/Passes/CompositePassBlur.h
--- a/Engine/Plugins/Compositing/Composite/Source/Composite/Public/Passes/CompositePassBlur.h
+++ b/Engine/Plugins/Compositing/Composite/Source/Composite/Public/Passes/CompositePassBlur.h
@@ -44,10 +44,6 @@ public:
 	/** Only blur the alpha channel, leaving RGB unchanged. */
 	UPROPERTY(EditAnywhere, BlueprintReadWrite, Interp, Category = "Composite", AdvancedDisplay)
 	bool bAlphaOnly = false;
-
-	/** When alpha-only, unpremultiply RGB before blur and re-premultiply after. Disable if input is not premultiplied. */
-	UPROPERTY(BlueprintReadWrite, EditAnywhere, Interp, Category = "Composite", AdvancedDisplay, meta = (EditCondition = "bAlphaOnly"))
-	bool bUnpremult = true;
 };
 
 #undef UE_API
diff --git a/Engine/Plugins/Compositing/Composite/Source/Composite/Public/Passes/CompositePassOpenColorIO.h b/Engine/Plugins/Compositing/Composite/Source/Composite/Public/Passes/CompositePassOpenColorIO.h
--- a/Engine/Plugins/Compositing/Composite/Source/Composite/Public/Passes/CompositePassOpenColorIO.h
+++ b/Engine/Plugins/Compositing/Composite/Source/Composite/Public/Passes/CompositePassOpenColorIO.h
@@ -30,6 +30,10 @@ public:
 	/** Color conversion settings. */
 	UPROPERTY(EditAnywhere, BlueprintReadWrite, Category = "Composite")
 	FOpenColorIOColorConversionSettings ColorConversionSettings;
+
+	/** Color is divided by its alpha channel before the color transform, and multiplied again after. */
+	UPROPERTY(BlueprintReadWrite, EditAnywhere, Category = "Composite")
+	bool bUnpremult;
 };
 
 #undef UE_API
