Composure: Exposed "ScaleFactor" parameter on composite mesh to facilitate scaling to Unreal units.

* Renamed from previous "ScaleToCentimeters" parameter name, since unreal units are configurable.

#virtualized
#Synced-CL 52383054

diff --git a/Engine/Plugins/Compositing/Composite/Source/Composite/Private/Components/CompositeDepthMeshComponent.cpp b/Engine/Plugins/Compositing/Composite/Source/Composite/Private/Components/CompositeDepthMeshComponent.cpp
--- a/Engine/Plugins/Compositing/Composite/Source/Composite/Private/Components/CompositeDepthMeshComponent.cpp
+++ b/Engine/Plugins/Compositing/Composite/Source/Composite/Private/Components/CompositeDepthMeshComponent.cpp
@@ -25,6 +25,7 @@ namespace UE
 		namespace Private
 		{
 			static FLazyName CompositeDepthTextureName(TEXT("CompositeDepthTexture"));
+			static FLazyName CompositeScaleFactorName(TEXT("ScaleFactor"));
 
 #if WITH_EDITOR
 			FString GetAssetPath(UCompositeDepthMeshComponent& InCompositeDepthMesh)
@@ -167,6 +168,7 @@ void UCompositeDepthMeshComponent::UpdateMaterial()
 		if (ensureMsgf(MID, TEXT("Expected a material instance dynamic")))
 		{
 			MID->SetTextureParameterValue(UE::Composite::Private::CompositeDepthTextureName, DepthTexture);
+			MID->SetScalarParameterValue(UE::Composite::Private::CompositeScaleFactorName, ScaleFactor);
 		}
 	}
 }
diff --git a/Engine/Plugins/Compositing/Composite/Source/Composite/Public/Components/CompositeDepthMeshComponent.h b/Engine/Plugins/Compositing/Composite/Source/Composite/Public/Components/CompositeDepthMeshComponent.h
--- a/Engine/Plugins/Compositing/Composite/Source/Composite/Public/Components/CompositeDepthMeshComponent.h
+++ b/Engine/Plugins/Compositing/Composite/Source/Composite/Public/Components/CompositeDepthMeshComponent.h
@@ -38,11 +38,15 @@ public:
 	UPROPERTY(EditAnywhere, Category = "Composite")
 	TObjectPtr<UTexture> DepthTexture;
 
+	/* Scale factor sent to material, to convert depth units into Unreal units. */
+	UPROPERTY(EditAnywhere, BlueprintReadOnly, Category = "Composite")
+	float ScaleFactor = 1.0f;
+
 protected:
 	UE_API virtual FBoxSphereBounds CalcBounds(const FTransform& LocalToWorld) const override;
 
 private:
-	/* Convert material to MID and update its depth texture parameter. */
+	/* Convert material to MID and update its parameters. */
 	void UpdateMaterial();
 
 #if WITH_EDITOR
