Composure: Renamed AlphaTransform pass to Translucency.

#Synced-CL 52419094

diff --git a/Engine/Plugins/Compositing/Composite/Config/DefaultComposite.ini b/Engine/Plugins/Compositing/Composite/Config/DefaultComposite.ini
--- a/Engine/Plugins/Compositing/Composite/Config/DefaultComposite.ini
+++ b/Engine/Plugins/Compositing/Composite/Config/DefaultComposite.ini
@@ -44,3 +44,4 @@
 +PropertyRedirects=(OldName="/Script/Composite.CompositePassCenteredScale.SourceAspectRatio", NewName="/Script/Composite.CompositePassCenteredScale.ContainerAspectRatio")
 +PropertyRedirects=(OldName="/Script/Composite.CompositePassCenteredScale.TargetAspectRatio", NewName="/Script/Composite.CompositePassCenteredScale.ContentAspectRatio")
 +PropertyRedirects=(OldName="/Script/Composite.CompositePassCenteredScale.ScaleFactor", NewName="/Script/Composite.CompositePassCenteredScale.ManualScale")
++ClassRedirects=(OldName="/Script/Composite.CompositePassAlphaTransform",NewName="/Script/Composite.CompositePassTranslucency")
diff --git a/Engine/Plugins/Compositing/Composite/Shaders/Private/CompositeAlphaTransform.usf b/Engine/Plugins/Compositing/Composite/Shaders/Private/CompositeAlphaTransform.usf
deleted file mode 100644
diff --git a/Engine/Plugins/Compositing/Composite/Shaders/Private/CompositeTranslucency.usf b/Engine/Plugins/Compositing/Composite/Shaders/Private/CompositeTranslucency.usf
new file mode 100644
--- /dev/null
+++ b/Engine/Plugins/Compositing/Composite/Shaders/Private/CompositeTranslucency.usf
@@ -0,0 +1,57 @@
+// Copyright Epic Games, Inc. All Rights Reserved.
+
+#include "/Engine/Private/Common.ush"
+#include "/Engine/Private/ScreenPass.ush"
+
+#include "CompositeCommon.ush"
+
+FScreenTransform SvPositionToInputTextureUV;
+Texture2D Input_Texture;
+SamplerState Input_Sampler;
+
+uint PremultOp;
+float Fade;
+uint OverrideOutputAlpha;
+
+#define PremultOp_None 0u
+#define PremultOp_Premultiply 1u
+#define PremultOp_Unpremultiply 2u
+
+#define OverrideOutput_None 0u
+#define OverrideOutput_Zero 1u
+#define OverrideOutput_One 2u
+
+void MainPS(
+    float4 SvPosition : SV_POSITION,
+    out float4 OutColor : SV_Target0
+    )
+{
+    float2 TextureUV = ApplyScreenTransform(SvPosition.xy, SvPositionToInputTextureUV);
+    float4 Color = Input_Texture.SampleLevel(Input_Sampler, TextureUV, 0);
+
+	if(PremultOp == PremultOp_Premultiply)
+	{
+		Color.rgb *= Color.a;
+		Color *= Fade;
+	}
+	else if(PremultOp == PremultOp_Unpremultiply)
+	{
+		Color.rgb *= SafeRcp(Color.a);
+		Color.a *= Fade;
+	}
+	else
+	{
+		Color *= Fade;
+	}
+
+	if(OverrideOutputAlpha == OverrideOutput_Zero)
+	{
+		Color.a = 0.0f;
+	}
+	else if(OverrideOutputAlpha == OverrideOutput_One)
+	{
+		Color.a = 1.0f;
+	}
+
+	OutColor = Color;
+}
diff --git a/Engine/Plugins/Compositing/Composite/Source/Composite/Private/Layers/CompositeLayerPlate.cpp b/Engine/Plugins/Compositing/Composite/Source/Composite/Private/Layers/CompositeLayerPlate.cpp
--- a/Engine/Plugins/Compositing/Composite/Source/Composite/Private/Layers/CompositeLayerPlate.cpp
+++ b/Engine/Plugins/Compositing/Composite/Source/Composite/Private/Layers/CompositeLayerPlate.cpp
@@ -17,7 +17,7 @@
 #include "Materials/Material.h"
 #include "Materials/MaterialInstanceDynamic.h"
 #include "Passes/CompositePassDistortion.h"
-#include "Passes/CompositePassAlphaTransform.h"
+#include "Passes/CompositePassTranslucency.h"
 #include "Profile/IMediaProfileManager.h"
 #include "Profile/MediaProfile.h"
 #include "Profile/MediaProfilePlaybackManager.h"
@@ -626,7 +626,7 @@ FCompositeCorePassProxy* UCompositeLayerPlate::GetProxy(const UE::CompositeCore:
 				FPassInputDecl AlphaPassInputDecl;
 				AlphaPassInputDecl.Set<FPassExternalResourceDesc>({ TextureId });
 
-				FAlphaTransformPassProxy* AlphaPassProxy = InFrameAllocator.Create<FAlphaTransformPassProxy>(FPassInputDeclArray{ AlphaPassInputDecl });
+				FTranslucencyPassProxy* AlphaPassProxy = InFrameAllocator.Create<FTranslucencyPassProxy>(FPassInputDeclArray{ AlphaPassInputDecl });
 				AlphaPassProxy->PremultOp = ECompositeAlphaPremultiplication::Unpremultiply;
 				AlphaPassProxy->OverrideOutputAlpha = ECompositeAlphaOverride::One;
 
diff --git a/Engine/Plugins/Compositing/Composite/Source/Composite/Private/Passes/CompositePassAlphaTransform.cpp b/Engine/Plugins/Compositing/Composite/Source/Composite/Private/Passes/CompositePassAlphaTransform.cpp
deleted file mode 100644
diff --git a/Engine/Plugins/Compositing/Composite/Source/Composite/Private/Passes/CompositePassTranslucency.cpp b/Engine/Plugins/Compositing/Composite/Source/Composite/Private/Passes/CompositePassTranslucency.cpp
new file mode 100644
--- /dev/null
+++ b/Engine/Plugins/Compositing/Composite/Source/Composite/Private/Passes/CompositePassTranslucency.cpp
@@ -0,0 +1,100 @@
+// Copyright Epic Games, Inc. All Rights Reserved.
+
+#include "Passes/CompositePassTranslucency.h"
+
+#include "Engine/Texture.h"
+#include "PixelShaderUtils.h"
+#include "RenderGraphBuilder.h"
+#include "RenderGraphUtils.h"
+#include "SceneView.h"
+#include "ScreenPass.h"
+
+DECLARE_GPU_STAT_NAMED(FCompositeTranslucency, TEXT("Composite.Translucency"));
+
+class FCompositePassTranslucencyShader : public FGlobalShader
+{
+	DECLARE_GLOBAL_SHADER(FCompositePassTranslucencyShader);
+	SHADER_USE_PARAMETER_STRUCT(FCompositePassTranslucencyShader, FGlobalShader);
+
+	BEGIN_SHADER_PARAMETER_STRUCT(FParameters, )
+		SHADER_PARAMETER_STRUCT(FScreenPassTextureInput, Input)
+		SHADER_PARAMETER(FScreenTransform, SvPositionToInputTextureUV)
+		SHADER_PARAMETER(uint32, PremultOp)
+		SHADER_PARAMETER(float, Fade)
+		SHADER_PARAMETER(uint32, OverrideOutputAlpha)
+		RENDER_TARGET_BINDING_SLOTS()
+	END_SHADER_PARAMETER_STRUCT()
+};
+IMPLEMENT_GLOBAL_SHADER(FCompositePassTranslucencyShader, "/Plugin/Composite/Private/CompositeTranslucency.usf", "MainPS", SF_Pixel);
+
+
+namespace UE
+{
+	namespace Composite
+	{
+		namespace Private
+		{
+			FPassTexture FTranslucencyPassProxy::Add(FRDGBuilder& GraphBuilder, const FSceneView& InView, const FPassInputArray& Inputs, const FPassContext& PassContext) const
+			{
+				RDG_EVENT_SCOPE_STAT(GraphBuilder, FCompositeTranslucency, "Composite.Translucency");
+
+				check(ValidateInputs(Inputs));
+
+				const UE::CompositeCore::FResourceMetadata& Metadata = Inputs[0].Metadata;
+				const FScreenPassTexture& InputTexture = Inputs[0].Texture;
+				FScreenPassRenderTarget Output = Inputs.OverrideOutput;
+
+				if (!Output.IsValid())
+				{
+					Output = CreateOutputRenderTarget(GraphBuilder, InView, PassContext.OutputViewRect, InputTexture.Texture->Desc, TEXT("TranslucencyCompositePass"));
+				}
+
+				FRHISamplerState* InputSampler = Metadata.GetSamplerState();
+				FGlobalShaderMap* GlobalShaderMap = GetGlobalShaderMap(InView.GetFeatureLevel());
+				FCompositePassTranslucencyShader::FParameters* Parameters = GraphBuilder.AllocParameters<FCompositePassTranslucencyShader::FParameters>();
+				Parameters->Input = GetScreenPassTextureInput(InputTexture, InputSampler);
+				Parameters->SvPositionToInputTextureUV = (
+					FScreenTransform::ChangeTextureBasisFromTo(FScreenPassTextureViewport(Output), FScreenTransform::ETextureBasis::TexelPosition, FScreenTransform::ETextureBasis::ViewportUV) *
+					FScreenTransform::ChangeTextureBasisFromTo(FScreenPassTextureViewport(InputTexture), FScreenTransform::ETextureBasis::ViewportUV, FScreenTransform::ETextureBasis::TextureUV));
+				Parameters->PremultOp = static_cast<uint32>(PremultOp);
+				Parameters->Fade = Fade;
+				Parameters->OverrideOutputAlpha = static_cast<uint32>(OverrideOutputAlpha);
+				Parameters->RenderTargets[0] = Output.GetRenderTargetBinding();
+
+				TShaderMapRef<FCompositePassTranslucencyShader> PixelShader(GlobalShaderMap);
+				FPixelShaderUtils::AddFullscreenPass(
+					GraphBuilder,
+					GlobalShaderMap,
+					RDG_EVENT_NAME("Composite.Translucency (%dx%d)", Output.ViewRect.Width(), Output.ViewRect.Height()),
+					PixelShader,
+					Parameters,
+					Output.ViewRect
+				);
+
+				return UE::CompositeCore::FPassTexture{ MoveTemp(Output), Metadata };
+			}
+		}
+	}
+}
+
+UCompositePassTranslucency::UCompositePassTranslucency(const FObjectInitializer& ObjectInitializer)
+	: Super(ObjectInitializer)
+	, PremultOp{ ECompositeAlphaPremultiplication::None }
+	, Fade{ 1.0f }
+	, OverrideOutputAlpha{ ECompositeAlphaOverride::None }
+{
+}
+
+UCompositePassTranslucency::~UCompositePassTranslucency() = default;
+
+FCompositeCorePassProxy* UCompositePassTranslucency::GetProxy(const UE::CompositeCore::FPassInputDecl& InputDecl, FCompositeTraversalContext& InContext, FSceneRenderingBulkObjectAllocator& InFrameAllocator) const
+{
+	using namespace UE::Composite::Private;
+
+	FTranslucencyPassProxy* Proxy = InFrameAllocator.Create<FTranslucencyPassProxy>(FPassInputDeclArray{ InputDecl });
+	Proxy->PremultOp = PremultOp;
+	Proxy->Fade = Fade;
+	Proxy->OverrideOutputAlpha = OverrideOutputAlpha;
+
+	return Proxy;
+}
diff --git a/Engine/Plugins/Compositing/Composite/Source/Composite/Public/Passes/CompositePassAlphaTransform.h b/Engine/Plugins/Compositing/Composite/Source/Composite/Public/Passes/CompositePassAlphaTransform.h
deleted file mode 100644
diff --git a/Engine/Plugins/Compositing/Composite/Source/Composite/Public/Passes/CompositePassTranslucency.h b/Engine/Plugins/Compositing/Composite/Source/Composite/Public/Passes/CompositePassTranslucency.h
new file mode 100644
--- /dev/null
+++ b/Engine/Plugins/Compositing/Composite/Source/Composite/Public/Passes/CompositePassTranslucency.h
@@ -0,0 +1,82 @@
+// Copyright Epic Games, Inc. All Rights Reserved.
+
+#pragma once
+
+#include "CompositePassBase.h"
+#include "Passes/CompositeCorePassProxy.h"
+
+#include "CompositePassTranslucency.generated.h"
+
+#define UE_API COMPOSITE_API
+
+/** Alpha premultiplication operation type. */
+UENUM(BlueprintType)
+enum class ECompositeAlphaPremultiplication : uint8
+{
+	None,
+	Premultiply,
+	Unpremultiply,
+};
+
+UENUM(BlueprintType)
+enum class ECompositeAlphaOverride : uint8
+{
+	None,
+	Zero,
+	One,
+};
+
+/** Translucency pass to modify alpha. */
+UCLASS(MinimalAPI, BlueprintType, Blueprintable, EditInlineNew, CollapseCategories, meta=(DisplayName="Translucency"))
+class UCompositePassTranslucency : public UCompositePassBase
+{
+	GENERATED_BODY()
+
+public:
+	/** Constructor */
+	UE_API UCompositePassTranslucency(const FObjectInitializer& ObjectInitializer);
+
+	/** Destructor */
+	UE_API ~UCompositePassTranslucency();
+
+	UE_API virtual FCompositeCorePassProxy* GetProxy(const UE::CompositeCore::FPassInputDecl& InputDecl, FCompositeTraversalContext& InContext, FSceneRenderingBulkObjectAllocator& InFrameAllocator) const override;
+
+public:
+	/** Alpha premultiplication operation. */
+	UPROPERTY(BlueprintReadWrite, EditAnywhere, Interp, Category = "Composite")
+	ECompositeAlphaPremultiplication PremultOp;
+
+	/** Fade the alpha from 0 to 1. Assumes pre-multiplied alpha, so both RGB and A are scaled together. */
+	UPROPERTY(BlueprintReadWrite, EditAnywhere, Interp, Category = "Composite", meta = (ClampMin = "0.0", ClampMax = "1.0", UIMin = "0.0", UIMax = "1.0"))
+	float Fade = 1.0f;
+
+	UPROPERTY()
+	ECompositeAlphaOverride OverrideOutputAlpha;
+};
+
+namespace UE
+{
+	namespace Composite
+	{
+		namespace Private
+		{
+			using namespace CompositeCore;
+
+			class FTranslucencyPassProxy : public FCompositeCorePassProxy
+			{
+			public:
+				IMPLEMENT_COMPOSITE_PASS(FTranslucencyPassProxy);
+
+				using FCompositeCorePassProxy::FCompositeCorePassProxy;
+
+				FPassTexture Add(FRDGBuilder& GraphBuilder, const FSceneView& InView, const FPassInputArray& Inputs, const FPassContext& PassContext) const override;
+
+				ECompositeAlphaPremultiplication PremultOp = ECompositeAlphaPremultiplication::None;
+				float Fade = 1.0f;
+				ECompositeAlphaOverride OverrideOutputAlpha = ECompositeAlphaOverride::None;
+			};
+		}
+	}
+}
+
+#undef UE_API
