Standardize Composite plugin pass shader class naming and fix GPU stat/RDG event strings.

- Standardize pass shader classes to FCompositePass{Name}Shader convention
- Fix stray CompositeCore.* GPU stat/RDG prefixes to Composite.* (SMAA, Denoise)

#Synced-CL 52419094

diff --git a/Engine/Plugins/Compositing/Composite/Source/Composite/Private/Passes/CompositePassCenteredScale.cpp b/Engine/Plugins/Compositing/Composite/Source/Composite/Private/Passes/CompositePassCenteredScale.cpp
--- a/Engine/Plugins/Compositing/Composite/Source/Composite/Private/Passes/CompositePassCenteredScale.cpp
+++ b/Engine/Plugins/Compositing/Composite/Source/Composite/Private/Passes/CompositePassCenteredScale.cpp
@@ -16,10 +16,10 @@
 
 DECLARE_GPU_STAT_NAMED(FCompositeCenteredScale, TEXT("Composite.CenteredScale"));
 
-class FCompositeCenteredScaleShader : public FGlobalShader
+class FCompositePassCenteredScaleShader : public FGlobalShader
 {
-	DECLARE_GLOBAL_SHADER(FCompositeCenteredScaleShader);
-	SHADER_USE_PARAMETER_STRUCT(FCompositeCenteredScaleShader, FGlobalShader);
+	DECLARE_GLOBAL_SHADER(FCompositePassCenteredScaleShader);
+	SHADER_USE_PARAMETER_STRUCT(FCompositePassCenteredScaleShader, FGlobalShader);
 
 	BEGIN_SHADER_PARAMETER_STRUCT(FParameters, )
 		SHADER_PARAMETER_STRUCT(FScreenPassTextureInput, Input)
@@ -29,7 +29,7 @@ class FCompositeCenteredScaleShader : public FGlobalShader
 		RENDER_TARGET_BINDING_SLOTS()
 	END_SHADER_PARAMETER_STRUCT()
 };
-IMPLEMENT_GLOBAL_SHADER(FCompositeCenteredScaleShader, "/Plugin/Composite/Private/CompositeCenteredScale.usf", "MainPS", SF_Pixel);
+IMPLEMENT_GLOBAL_SHADER(FCompositePassCenteredScaleShader, "/Plugin/Composite/Private/CompositeCenteredScale.usf", "MainPS", SF_Pixel);
 
 
 namespace UE
@@ -60,7 +60,7 @@ namespace UE
 				}
 
 				FGlobalShaderMap* GlobalShaderMap = GetGlobalShaderMap(InView.GetFeatureLevel());
-				FCompositeCenteredScaleShader::FParameters* Parameters = GraphBuilder.AllocParameters<FCompositeCenteredScaleShader::FParameters>();
+				FCompositePassCenteredScaleShader::FParameters* Parameters = GraphBuilder.AllocParameters<FCompositePassCenteredScaleShader::FParameters>();
 				FRHISamplerState* InputSampler = Metadata.GetSamplerState();
 				Parameters->Input = GetScreenPassTextureInput(Input, InputSampler);
 				Parameters->SvPositionToOutputViewportUV = FScreenTransform::ChangeTextureBasisFromTo(
@@ -74,7 +74,7 @@ namespace UE
 				Parameters->ScaleUV = ScaleUV;
 				Parameters->RenderTargets[0] = Output.GetRenderTargetBinding();
 
-				TShaderMapRef<FCompositeCenteredScaleShader> PixelShader(GlobalShaderMap);
+				TShaderMapRef<FCompositePassCenteredScaleShader> PixelShader(GlobalShaderMap);
 				FPixelShaderUtils::AddFullscreenPass(
 					GraphBuilder,
 					GlobalShaderMap,
diff --git a/Engine/Plugins/Compositing/Composite/Source/Composite/Private/Passes/CompositePassColorGrading.cpp b/Engine/Plugins/Compositing/Composite/Source/Composite/Private/Passes/CompositePassColorGrading.cpp
--- a/Engine/Plugins/Compositing/Composite/Source/Composite/Private/Passes/CompositePassColorGrading.cpp
+++ b/Engine/Plugins/Compositing/Composite/Source/Composite/Private/Passes/CompositePassColorGrading.cpp
@@ -19,10 +19,10 @@ BEGIN_SHADER_PARAMETER_STRUCT(FColorGradingPerRangeParameters, )
 	SHADER_PARAMETER(FVector4f, Offset)
 END_SHADER_PARAMETER_STRUCT()
 
-class FCompositeColorGradingShader : public FGlobalShader
+class FCompositePassColorGradingShader : public FGlobalShader
 {
-	DECLARE_GLOBAL_SHADER(FCompositeColorGradingShader);
-	SHADER_USE_PARAMETER_STRUCT(FCompositeColorGradingShader, FGlobalShader);
+	DECLARE_GLOBAL_SHADER(FCompositePassColorGradingShader);
+	SHADER_USE_PARAMETER_STRUCT(FCompositePassColorGradingShader, FGlobalShader);
 
 	BEGIN_SHADER_PARAMETER_STRUCT(FParameters, )
 		SHADER_PARAMETER_STRUCT(FScreenPassTextureInput, Input)
@@ -55,7 +55,7 @@ class FCompositeColorGradingShader : public FGlobalShader
 		return Output;
 	}
 };
-IMPLEMENT_GLOBAL_SHADER(FCompositeColorGradingShader, "/Plugin/Composite/Private/CompositeColorGrading.usf", "MainPS", SF_Pixel);
+IMPLEMENT_GLOBAL_SHADER(FCompositePassColorGradingShader, "/Plugin/Composite/Private/CompositeColorGrading.usf", "MainPS", SF_Pixel);
 
 namespace UE::Composite::Private
 {
@@ -81,14 +81,14 @@ namespace UE::Composite::Private
 		}
 
 		FGlobalShaderMap* GlobalShaderMap = GetGlobalShaderMap(InView.GetFeatureLevel());
-		FCompositeColorGradingShader::FParameters* Parameters = GraphBuilder.AllocParameters<FCompositeColorGradingShader::FParameters>();
+		FCompositePassColorGradingShader::FParameters* Parameters = GraphBuilder.AllocParameters<FCompositePassColorGradingShader::FParameters>();
 		FRHISamplerState* InputSampler = Metadata.GetSamplerState();
 		Parameters->Input = GetScreenPassTextureInput(Input, InputSampler);
 		Parameters->WorkingColorSpace = GDefaultWorkingColorSpaceUniformBuffer.GetUniformBufferRef();
-		Parameters->Global = FCompositeColorGradingShader::ToParameters(ColorGradingSettings.Global);
-		Parameters->Shadows = FCompositeColorGradingShader::ToParameters(ColorGradingSettings.Shadows);
-		Parameters->Midtones = FCompositeColorGradingShader::ToParameters(ColorGradingSettings.Midtones);
-		Parameters->Highlights = FCompositeColorGradingShader::ToParameters(ColorGradingSettings.Highlights);
+		Parameters->Global = FCompositePassColorGradingShader::ToParameters(ColorGradingSettings.Global);
+		Parameters->Shadows = FCompositePassColorGradingShader::ToParameters(ColorGradingSettings.Shadows);
+		Parameters->Midtones = FCompositePassColorGradingShader::ToParameters(ColorGradingSettings.Midtones);
+		Parameters->Highlights = FCompositePassColorGradingShader::ToParameters(ColorGradingSettings.Highlights);
 		Parameters->bUnpremult = bUnpremult;
 		Parameters->bInvertedAlpha = static_cast<uint32>(Metadata.bInvertedAlpha);
 		Parameters->ShadowsMax = ColorGradingSettings.ShadowsMax;
@@ -102,7 +102,7 @@ namespace UE::Composite::Private
 			FScreenTransform::ChangeTextureBasisFromTo(FScreenPassTextureViewport(Input), FScreenTransform::ETextureBasis::ViewportUV, FScreenTransform::ETextureBasis::TextureUV));
 		Parameters->RenderTargets[0] = Output.GetRenderTargetBinding();
 
-		TShaderMapRef<FCompositeColorGradingShader> PixelShader(GlobalShaderMap);
+		TShaderMapRef<FCompositePassColorGradingShader> PixelShader(GlobalShaderMap);
 		FPixelShaderUtils::AddFullscreenPass(
 			GraphBuilder,
 			GlobalShaderMap,
diff --git a/Engine/Plugins/Compositing/Composite/Source/Composite/Private/Passes/CompositePassColorKeyer.cpp b/Engine/Plugins/Compositing/Composite/Source/Composite/Private/Passes/CompositePassColorKeyer.cpp
--- a/Engine/Plugins/Compositing/Composite/Source/Composite/Private/Passes/CompositePassColorKeyer.cpp
+++ b/Engine/Plugins/Compositing/Composite/Source/Composite/Private/Passes/CompositePassColorKeyer.cpp
@@ -104,7 +104,7 @@ namespace UE
 				TShaderMapRef<FCompositePassDenoiseShader> ComputeShader(GlobalShaderMap);
 				FComputeShaderUtils::AddPass(
 					GraphBuilder,
-					RDG_EVENT_NAME("CompositeCore.Denoise (%dx%d)", TextureSize.X, TextureSize.Y),
+					RDG_EVENT_NAME("Composite.Denoise (%dx%d)", TextureSize.X, TextureSize.Y),
 					ERDGPassFlags::Compute,
 					ComputeShader,
 					PassParameters,
diff --git a/Engine/Plugins/Compositing/Composite/Source/Composite/Private/Passes/CompositePassDistortion.cpp b/Engine/Plugins/Compositing/Composite/Source/Composite/Private/Passes/CompositePassDistortion.cpp
--- a/Engine/Plugins/Compositing/Composite/Source/Composite/Private/Passes/CompositePassDistortion.cpp
+++ b/Engine/Plugins/Compositing/Composite/Source/Composite/Private/Passes/CompositePassDistortion.cpp
@@ -15,10 +15,10 @@
 #include "PostProcess/PostProcessMaterialInputs.h"
 #include "SceneView.h"
 
-class FCompositeDistortionShader : public FGlobalShader
+class FCompositePassDistortionShader : public FGlobalShader
 {
-	DECLARE_GLOBAL_SHADER(FCompositeDistortionShader);
-	SHADER_USE_PARAMETER_STRUCT(FCompositeDistortionShader, FGlobalShader);
+	DECLARE_GLOBAL_SHADER(FCompositePassDistortionShader);
+	SHADER_USE_PARAMETER_STRUCT(FCompositePassDistortionShader, FGlobalShader);
 
 	BEGIN_SHADER_PARAMETER_STRUCT(FParameters, )
 		SHADER_PARAMETER_STRUCT(FScreenPassTextureInput, Input)
@@ -31,7 +31,7 @@ class FCompositeDistortionShader : public FGlobalShader
 		RENDER_TARGET_BINDING_SLOTS()
 	END_SHADER_PARAMETER_STRUCT()
 };
-IMPLEMENT_GLOBAL_SHADER(FCompositeDistortionShader, "/Plugin/Composite/Private/CompositeDistortion.usf", "MainPS", SF_Pixel);
+IMPLEMENT_GLOBAL_SHADER(FCompositePassDistortionShader, "/Plugin/Composite/Private/CompositeDistortion.usf", "MainPS", SF_Pixel);
 
 DECLARE_GPU_STAT_NAMED(FCompositeDistortion, TEXT("Composite.Distortion"));
 
@@ -94,7 +94,7 @@ public:
 			FGlobalShaderMap* GlobalShaderMap = GetGlobalShaderMap(InView.GetFeatureLevel());
 			FScreenPassTextureViewport Viewport = FScreenPassTextureViewport(Output);
 
-			FCompositeDistortionShader::FParameters* PassParameters = GraphBuilder.AllocParameters<FCompositeDistortionShader::FParameters>();
+			FCompositePassDistortionShader::FParameters* PassParameters = GraphBuilder.AllocParameters<FCompositePassDistortionShader::FParameters>();
 			// Always use bilinear: distortion resamples through a displacement LUT, mapping texels to fractional coordinates.
 			PassParameters->Input = GetScreenPassTextureInput(Input, TStaticSamplerState<SF_Bilinear>::GetRHI());
 			PassParameters->Output = GetScreenPassTextureViewportParameters(FScreenPassTextureViewport(Output));
@@ -126,7 +126,7 @@ public:
 
 			PassParameters->RenderTargets[0] = Output.GetRenderTargetBinding();
 
-			TShaderMapRef<FCompositeDistortionShader> PixelShader(GlobalShaderMap);
+			TShaderMapRef<FCompositePassDistortionShader> PixelShader(GlobalShaderMap);
 			
 			FPixelShaderUtils::AddFullscreenPass(
 				GraphBuilder,
diff --git a/Engine/Plugins/Compositing/Composite/Source/Composite/Private/Passes/CompositePassSMAA.cpp b/Engine/Plugins/Compositing/Composite/Source/Composite/Private/Passes/CompositePassSMAA.cpp
--- a/Engine/Plugins/Compositing/Composite/Source/Composite/Private/Passes/CompositePassSMAA.cpp
+++ b/Engine/Plugins/Compositing/Composite/Source/Composite/Private/Passes/CompositePassSMAA.cpp
@@ -9,7 +9,7 @@
 #include "ScreenPass.h"
 #include "PostProcess/SubpixelMorphologicalAA.h"
 
-DECLARE_GPU_STAT_NAMED(FCompositeCoreSMAA, TEXT("CompositeCore.SMAA"));
+DECLARE_GPU_STAT_NAMED(FCompositeSMAA, TEXT("Composite.SMAA"));
 
 namespace UE
 {
@@ -26,7 +26,7 @@ namespace UE
 
 			FPassTexture Add(FRDGBuilder& GraphBuilder, const FSceneView& InView, const FPassInputArray& Inputs, const FPassContext& PassContext) const override
 			{
-				RDG_EVENT_SCOPE_STAT(GraphBuilder, FCompositeCoreSMAA, "CompositeCore.SMAA");
+				RDG_EVENT_SCOPE_STAT(GraphBuilder, FCompositeSMAA, "Composite.SMAA");
 
 				check(ValidateInputs(Inputs));
 
