Composure: Automatic overscan removal on plate layer.

- Add bRemoveOverscan (default true, AdvancedDisplay) to plate layer with an internal CenteredScale pass that automatically counteracts camera overscan on the media texture. This is preferable over asking users to add a centered scale pass.
- Remove ECompositePassOverscanUncropMode enum and manual overscan; replace with a simple bRemoveOverscan bool on the CenteredScale pass. Manual scaling can still be done with the regular user-exposed scale pass.

#Synced-CL 52419094

diff --git a/Engine/Plugins/Compositing/Composite/Source/Composite/Private/Layers/CompositeLayerPlate.cpp b/Engine/Plugins/Compositing/Composite/Source/Composite/Private/Layers/CompositeLayerPlate.cpp
--- a/Engine/Plugins/Compositing/Composite/Source/Composite/Private/Layers/CompositeLayerPlate.cpp
+++ b/Engine/Plugins/Compositing/Composite/Source/Composite/Private/Layers/CompositeLayerPlate.cpp
@@ -18,6 +18,7 @@
 #include "Materials/MaterialInstanceDynamic.h"
 #include "Passes/CompositePassDistortion.h"
 #include "Passes/CompositePassTranslucency.h"
+#include "Passes/CompositePassCenteredScale.h"
 #include "Profile/IMediaProfileManager.h"
 #include "Profile/MediaProfile.h"
 #include "Profile/MediaProfilePlaybackManager.h"
@@ -151,6 +152,11 @@ UCompositeLayerPlate::UCompositeLayerPlate(const FObjectInitializer& ObjectIniti
 	// Default undistort pass for plate media going into the scene
 	ScenePasses.Add(CreateDefaultSubobject<UCompositePassDistortion>(TEXT("CompositeUndistortScenePass")));
 
+	// Internal pass to automatically remove overscan
+	OverscanPass = CreateDefaultSubobject<UCompositePassCenteredScale>(TEXT("CompositeOverscanPass"));
+	OverscanPass->ScaleMode = ECompositePassScaleMode::None;
+	OverscanPass->bRemoveOverscan = true;
+
 #if WITH_EDITOR
 	// If we are not the class default object...
 	if (!HasAnyFlags(RF_ClassDefaultObject))
@@ -182,7 +188,8 @@ void UCompositeLayerPlate::Tick(float DeltaTime)
 	TryRegisterCompositePrimitives(GetPrimitives());
 
 	// Constantly check if passes have changed to manage render targets & texture bindings on composite meshes.
-	const int32 NumValidMediaPasses = GetValidPassesNum(MediaPasses);
+	const int32 NumValidMediaPasses = GetValidPassesNum(MediaPasses)
+		+ ((bRemoveOverscan && IsValid(OverscanPass) && OverscanPass->GetIsActive()) ? 1 : 0);
 	const int32 NumValidScenePasses = GetValidPassesNum(ScenePasses);
 	bool bPassesHaveChanged = false;
 
@@ -590,14 +597,35 @@ FCompositeCorePassProxy* UCompositeLayerPlate::GetProxy(const UE::CompositeCore:
 		TextureId = InContext.FindOrCreateExternalTexture(Texture, Metadata);
 		
 		const ResourceId OriginalTextureId = TextureId;
+		
+		const bool bImplicitOverscanRemoval = bRemoveOverscan && IsValid(OverscanPass) && OverscanPass->GetIsActive();
 
-		// Add media texture pre-processing passes
-		TextureId = AddPreprocessingPasses(InContext, InFrameAllocator, MediaPasses, TextureId, OriginalTextureId,
-			[this]() -> TObjectPtr<UTextureRenderTarget2D>
-			{
-				return GetOrCreateRenderTarget(MediaRenderTarget);
-			}
-		);
+		if (bImplicitOverscanRemoval)
+		{
+			// Build combined media pass array: overscan (implicit) + user passes.
+			// Passes iterate backwards (highest index first), so overscan at the end executes first.
+			TArray<TObjectPtr<UCompositePassBase>> AllMediaPasses;
+			AllMediaPasses.Reserve(MediaPasses.Num() + 1);
+			AllMediaPasses.Append(MediaPasses);
+			AllMediaPasses.Add(OverscanPass);
+
+			TextureId = AddPreprocessingPasses(InContext, InFrameAllocator, AllMediaPasses, TextureId, OriginalTextureId,
+				[this]() -> TObjectPtr<UTextureRenderTarget2D>
+				{
+					return GetOrCreateRenderTarget(MediaRenderTarget);
+				}
+			);
+		}
+		else
+		{
+			// Add media texture pre-processing passes
+			TextureId = AddPreprocessingPasses(InContext, InFrameAllocator, MediaPasses, TextureId, OriginalTextureId,
+				[this]() -> TObjectPtr<UTextureRenderTarget2D>
+				{
+					return GetOrCreateRenderTarget(MediaRenderTarget);
+				}
+			);
+		}
 
 		// Add scene-only passes. We ignore returned resource since we don't refer to it in post-processing.
 		AddPreprocessingPasses(InContext, InFrameAllocator, ScenePasses, TextureId, OriginalTextureId,
diff --git a/Engine/Plugins/Compositing/Composite/Source/Composite/Private/Passes/CompositePassCenteredScale.cpp b/Engine/Plugins/Compositing/Composite/Source/Composite/Private/Passes/CompositePassCenteredScale.cpp
--- a/Engine/Plugins/Compositing/Composite/Source/Composite/Private/Passes/CompositePassCenteredScale.cpp
+++ b/Engine/Plugins/Compositing/Composite/Source/Composite/Private/Passes/CompositePassCenteredScale.cpp
@@ -100,8 +100,7 @@ UCompositePassCenteredScale::UCompositePassCenteredScale(const FObjectInitialize
 	, bScaleSingleAxis(true)
 	, UnpadPixels(FIntPoint::ZeroValue)
 	, ManualScale(1.0f, 1.0f)
-	, OverscanUncropMode(ECompositePassOverscanUncropMode::Automatic)
-	, Overscan(0.0f)
+	, bRemoveOverscan(true)
 {
 }
 
@@ -126,7 +125,7 @@ bool UCompositePassCenteredScale::GetIsActive() const
 	CachedScaleUV = CalculateScale();
 
 	return GetIsEnabled()
-		&& (ScaleMode != ECompositePassScaleMode::None || OverscanUncropMode != ECompositePassOverscanUncropMode::None)
+		&& (ScaleMode != ECompositePassScaleMode::None || bRemoveOverscan)
 		&& !CachedScaleUV.GetValue().Equals(FVector2f::One(), UE_SMALL_NUMBER);
 }
 
@@ -232,23 +231,9 @@ FVector2f UCompositePassCenteredScale::CalculateScale() const
 		break;
 	};
 
-	switch (OverscanUncropMode)
+	if (bRemoveOverscan && CameraComponent && CameraComponent->bCropOverscan)
 	{
-	case ECompositePassOverscanUncropMode::None:
-		break;
-	case ECompositePassOverscanUncropMode::Automatic:
-		if (CameraComponent)
-		{
-			ScaleUV *= 1.0f + CameraComponent->Overscan;
-		}
-		break;
-	case ECompositePassOverscanUncropMode::Manual:
-		ScaleUV *= 1.0f + Overscan;
-		break;
-
-	default:
-		checkNoEntry();
-		break;
+		ScaleUV *= 1.0f + CameraComponent->Overscan;
 	}
 
 	return ScaleUV;
diff --git a/Engine/Plugins/Compositing/Composite/Source/Composite/Public/Layers/CompositeLayerPlate.h b/Engine/Plugins/Compositing/Composite/Source/Composite/Public/Layers/CompositeLayerPlate.h
--- a/Engine/Plugins/Compositing/Composite/Source/Composite/Public/Layers/CompositeLayerPlate.h
+++ b/Engine/Plugins/Compositing/Composite/Source/Composite/Public/Layers/CompositeLayerPlate.h
@@ -14,6 +14,7 @@ class UPrimitiveComponent;
 class UTexture;
 class FObjectPreSaveContext;
 class UTextureRenderTarget2D;
+class UCompositePassCenteredScale;
 
 /** Plate sampling & rendering mode. */
 UENUM(BlueprintType)
@@ -195,6 +196,14 @@ public:
 	UPROPERTY(EditAnywhere, BlueprintReadWrite, Category = "Composite", AdvancedDisplay)
 	bool bBlendInDisplaySpace = false;
 
+	/**
+	 * Automatically remove overscan from the media texture by reading the camera's overscan value.
+	 * This inserts an implicit media pass that scales the footage to counteract the overscan expansion.
+	 * Has no effect when the camera overscan is zero.
+	 */
+	UPROPERTY(EditAnywhere, BlueprintReadWrite, Category = "Composite", AdvancedDisplay)
+	bool bRemoveOverscan = true;
+
 private:
 	/** Condtionally register composite mesh primitives: active holdout and add to the scene view extension custom render pass. */
 	void TryRegisterCompositePrimitives(const TArray<UPrimitiveComponent*>& InPrimitives, const UWorld* SourceWorld = nullptr) const;
@@ -281,6 +290,10 @@ private:
 	UPROPERTY(Transient, NonTransactional)
 	mutable TObjectPtr<UTextureRenderTarget2D> SceneRenderTarget;
 
+	/** Internal pass to automatically remove overscan from the media texture. Not serialized. */
+	UPROPERTY(Transient, NonTransactional)
+	TObjectPtr<UCompositePassCenteredScale> OverscanPass;
+
 	/** Cached number of valid media passes. */
 	int32 CachedValidMediaPasses = 0;
 
diff --git a/Engine/Plugins/Compositing/Composite/Source/Composite/Public/Passes/CompositePassCenteredScale.h b/Engine/Plugins/Compositing/Composite/Source/Composite/Public/Passes/CompositePassCenteredScale.h
--- a/Engine/Plugins/Compositing/Composite/Source/Composite/Public/Passes/CompositePassCenteredScale.h
+++ b/Engine/Plugins/Compositing/Composite/Source/Composite/Public/Passes/CompositePassCenteredScale.h
@@ -28,20 +28,6 @@ enum class ECompositePassScaleMode : uint8
 	Manual,
 };
 
-/** Overscan uncrop derivation mode. */
-UENUM(BlueprintType)
-enum class ECompositePassOverscanUncropMode : uint8
-{
-	/** Disabled */
-	None,
-
-	/** Automatically derive the overscan crop factor from the parent composite actor camera reference. */
-	Automatic,
-
-	/** Manually define the overscan crop factor.*/
-	Manual,
-};
-
 /** Convenience pass to re-scale footage (with black bars) inside its container texture. */
 UCLASS(MinimalAPI, BlueprintType, Blueprintable, EditInlineNew, CollapseCategories, meta=(DisplayName="Centered Scale"))
 class UCompositePassCenteredScale : public UCompositePassBase
@@ -88,13 +74,9 @@ public:
 	UPROPERTY(EditAnywhere, BlueprintReadWrite, Interp, Category = "Composite", meta = (EditCondition = "ScaleMode == ECompositePassScaleMode::Manual", EditConditionHides))
 	FVector2f ManualScale;
 
-	/** Uncrop calculation mode used to uncrop a viewport cropped in by overscan. */
-	UPROPERTY(EditAnywhere, BlueprintReadWrite, Category = "Composite")
-	ECompositePassOverscanUncropMode OverscanUncropMode;
-
-	/** Manual overscan to uncrop, with values from 0.0 to 1.0 matching the source camera overscan. */
-	UPROPERTY(EditAnywhere, BlueprintReadWrite, Interp, Category = "Composite", meta = (EditCondition = "OverscanUncropMode == ECompositePassOverscanUncropMode::Manual", EditConditionHides, ClampMin = "0.0", ClampMax = "1.0", UIMin = "0.0", UIMax = "1.0"))
-	float Overscan;
+	/** Automatically remove overscan by reading the camera's overscan value. */
+	UPROPERTY(BlueprintReadWrite, Category = "Composite")
+	bool bRemoveOverscan;
 
 private:
 	/** Cached scale UV computed in GetIsActive() and reused in GetProxy(). Reset after consumption. */
