Composure: Add uniform editor customizations and standardize actor picker filtering

- Add PlanarReflection detail customization with actor picker and pass tree panel
- Add TCompositeLayerPassListOwner<T> template replacing per-layer PassListOwner boilerplate
- Add TCompositeLayerPassesCustomization<T> template for pass-panel-only layers (MainRender, Processing)
- Consolidate actor picker filtering into SCompositeActorPickerTable::MakeDefaultSceneOutlinerFilters
- All actor pickers now check IsAllowedPrimitiveClass and exclude ACompositeMeshActor by default
- Plate and SceneCapture opt out of ACompositeMeshActor exclusion via constexpr bool
- Fix RemovePasses using ArrayAdd instead of ArrayRemove in pass list owners
- Delete CompositeSceneCapturePassListOwner (replaced by template)

#Synced-CL 52419094
#Synced-CL 52419094

diff --git a/Engine/Plugins/Compositing/Composite/Source/Composite/Public/Layers/CompositeLayerPlanarReflection.h b/Engine/Plugins/Compositing/Composite/Source/Composite/Public/Layers/CompositeLayerPlanarReflection.h
--- a/Engine/Plugins/Compositing/Composite/Source/Composite/Public/Layers/CompositeLayerPlanarReflection.h
+++ b/Engine/Plugins/Compositing/Composite/Source/Composite/Public/Layers/CompositeLayerPlanarReflection.h
@@ -108,6 +108,8 @@ private:
 	/** Convenience toggle to render the planar reflection as a fast/minimal Custom Render Pass, inlined in the main render & without support for lighting. */
 	UPROPERTY(EditAnywhere, BlueprintGetter = IsCustomRenderPass, BlueprintSetter = SetCustomRenderPass, Category = "Composite", AdvancedDisplay, meta = (AllowPrivateAccess = true))
 	bool bCustomRenderPass;
+
+	friend class FCompositeLayerPlanarReflectionCustomization;
 };
 
 #undef UE_API
diff --git a/Engine/Plugins/Compositing/Composite/Source/CompositeEditor/Private/CompositeEditorModule.cpp b/Engine/Plugins/Compositing/Composite/Source/CompositeEditor/Private/CompositeEditorModule.cpp
--- a/Engine/Plugins/Compositing/Composite/Source/CompositeEditor/Private/CompositeEditorModule.cpp
+++ b/Engine/Plugins/Compositing/Composite/Source/CompositeEditor/Private/CompositeEditorModule.cpp
@@ -26,14 +26,19 @@
 #include "Customizations/CompositeActorCustomization.h"
 #include "Customizations/CompositeViewProjectionComponentCustomization.h"
 #include "Customizations/CompositeLayerBaseCustomization.h"
+#include "Customizations/CompositeLayerPassesCustomization.h"
 #include "Customizations/CompositeLayerPlateCustomization.h"
 #include "Customizations/CompositeLayerSceneCaptureCustomization.h"
+#include "Customizations/CompositeLayerPlanarReflectionCustomization.h"
 #include "Customizations/CompositeLayerShadowReflectionCustomization.h"
 #include "Customizations/CompositeLayerSingleLightShadowCustomization.h"
 #include "Customizations/CompositeDepthMeshComponentCustomization.h"
 #include "Customizations/CompositePassColorKeyerCustomization.h"
 #include "Layers/CompositeLayerBase.h"
+#include "Layers/CompositeLayerMainRender.h"
 #include "Layers/CompositeLayerPlate.h"
+#include "Layers/CompositeLayerProcessing.h"
+#include "Layers/CompositeLayerPlanarReflection.h"
 #include "Layers/CompositeLayerSceneCapture.h"
 #include "Layers/CompositeLayerShadowReflection.h"
 #include "Layers/CompositeLayerSingleLightShadow.h"
@@ -229,6 +234,21 @@ void FCompositeEditorModule::RegisterCustomizations()
 			FOnGetDetailCustomizationInstance::CreateStatic(&FCompositeLayerSingleLightShadowCustomization::MakeInstance)
 		);
 
+	PropertyModule.RegisterCustomClassLayout(
+		UCompositeLayerPlanarReflection::StaticClass()->GetFName(),
+		FOnGetDetailCustomizationInstance::CreateStatic(&FCompositeLayerPlanarReflectionCustomization::MakeInstance)
+	);
+
+	PropertyModule.RegisterCustomClassLayout(
+		UCompositeLayerMainRender::StaticClass()->GetFName(),
+		FOnGetDetailCustomizationInstance::CreateStatic(&TCompositeLayerPassesCustomization<UCompositeLayerMainRender>::MakeInstance)
+	);
+
+	PropertyModule.RegisterCustomClassLayout(
+		UCompositeLayerProcessing::StaticClass()->GetFName(),
+		FOnGetDetailCustomizationInstance::CreateStatic(&TCompositeLayerPassesCustomization<UCompositeLayerProcessing>::MakeInstance)
+	);
+
 	PropertyModule.RegisterCustomClassLayout(
 		ACompositeActor::StaticClass()->GetFName(),
 		FOnGetDetailCustomizationInstance::CreateStatic(&FCompositeActorCustomization::MakeInstance)
@@ -262,6 +282,9 @@ void FCompositeEditorModule::UnregisterCustomizations()
 			PropertyModule->UnregisterCustomClassLayout(UCompositeLayerShadowReflection::StaticClass()->GetFName());
 			PropertyModule->UnregisterCustomClassLayout(UCompositeLayerSceneCapture::StaticClass()->GetFName());
 			PropertyModule->UnregisterCustomClassLayout(UCompositeLayerSingleLightShadow::StaticClass()->GetFName());
+			PropertyModule->UnregisterCustomClassLayout(UCompositeLayerPlanarReflection::StaticClass()->GetFName());
+			PropertyModule->UnregisterCustomClassLayout(UCompositeLayerMainRender::StaticClass()->GetFName());
+			PropertyModule->UnregisterCustomClassLayout(UCompositeLayerProcessing::StaticClass()->GetFName());
 			PropertyModule->UnregisterCustomClassLayout(UCompositeViewProjectionComponent::StaticClass()->GetFName());
 			PropertyModule->UnregisterCustomClassLayout(UCompositeDepthMeshComponent::StaticClass()->GetFName());
 			PropertyModule->UnregisterCustomClassLayout(UCompositePassColorKeyer::StaticClass()->GetFName());
diff --git a/Engine/Plugins/Compositing/Composite/Source/CompositeEditor/Private/Customizations/CompositeLayerPassesCustomization.h b/Engine/Plugins/Compositing/Composite/Source/CompositeEditor/Private/Customizations/CompositeLayerPassesCustomization.h
new file mode 100644
--- /dev/null
+++ b/Engine/Plugins/Compositing/Composite/Source/CompositeEditor/Private/Customizations/CompositeLayerPassesCustomization.h
@@ -0,0 +1,107 @@
+// Copyright Epic Games, Inc. All Rights Reserved.
+
+#pragma once
+
+#include "DetailCategoryBuilder.h"
+#include "DetailLayoutBuilder.h"
+#include "DetailWidgetRow.h"
+#include "IDetailCustomization.h"
+#include "IDetailGroup.h"
+#include "IPropertyUtilities.h"
+#include "UI/CompositeLayerPassListOwner.h"
+#include "UI/SCompositePassTreePanel.h"
+#include "Widgets/Text/STextBlock.h"
+
+/**
+ * Generic detail customization for layer types that only have a LayerPasses array (no actor picker).
+ * Covers UCompositeLayerMainRender and UCompositeLayerProcessing.
+ */
+template<typename LayerClass>
+class TCompositeLayerPassesCustomization : public IDetailCustomization
+{
+public:
+	static TSharedRef<IDetailCustomization> MakeInstance()
+	{
+		return MakeShared<TCompositeLayerPassesCustomization<LayerClass>>();
+	}
+
+	virtual void CustomizeDetails(const TSharedPtr<IDetailLayoutBuilder>& DetailBuilder) override
+	{
+		CachedDetailBuilder = DetailBuilder;
+		IDetailCustomization::CustomizeDetails(DetailBuilder);
+	}
+
+	virtual void CustomizeDetails(IDetailLayoutBuilder& DetailLayout) override
+	{
+		DetailLayout.HideProperty(GET_MEMBER_NAME_CHECKED(LayerClass, LayerPasses));
+
+		TArray<TWeakObjectPtr<LayerClass>> Objects = DetailLayout.GetObjectsOfTypeBeingCustomized<LayerClass>();
+
+		IDetailCategoryBuilder& LayerCategory = DetailLayout.EditCategory("Composite", LayerClass::StaticClass()->GetDisplayNameText());
+
+		TArray<TSharedRef<IPropertyHandle>> LayerProperties;
+		LayerCategory.GetDefaultProperties(LayerProperties, true, false);
+		for (const TSharedRef<IPropertyHandle>& Property : LayerProperties)
+		{
+			if (!Property->IsCustomized())
+			{
+				LayerCategory.AddProperty(Property);
+			}
+		}
+
+		if (Objects.Num() == 1 && Objects[0].IsValid())
+		{
+			LayerClass* Layer = Cast<LayerClass>(Objects[0].Get());
+
+			IDetailGroup& PassesGroup = LayerCategory.AddGroup("Passes", NSLOCTEXT("CompositeLayerPassesCustomization", "PassesGroupName", "Passes"), false, true);
+
+			PassesGroup.AddWidgetRow()
+			.WholeRowContent()
+			[
+				SAssignNew(PassPanel, SCompositePassTreePanel, MakeShared<TCompositeLayerPassListOwner<LayerClass>>(Layer))
+				.OnLayoutSizeChanged(this, &TCompositeLayerPassesCustomization::OnLayoutSizeChanged)
+			];
+		}
+		else
+		{
+			LayerCategory.AddCustomRow(NSLOCTEXT("CompositeLayerPassesCustomization", "PassesGroupName", "Passes"))
+			.NameContent()
+			[
+				SNew(STextBlock)
+				.Text(NSLOCTEXT("CompositeLayerPassesCustomization", "PassesGroupName", "Passes"))
+				.Font(IDetailLayoutBuilder::GetDetailFont())
+			]
+			.ValueContent()
+			[
+				SNew(STextBlock)
+				.Text(NSLOCTEXT("CompositeLayerPassesCustomization", "MultipleValues", "Multiple Values"))
+				.Font(IDetailLayoutBuilder::GetDetailFont())
+			];
+		}
+	}
+
+	bool IsCustomizingObject(UObject* InObject) const
+	{
+		if (TSharedPtr<IDetailLayoutBuilder> Builder = CachedDetailBuilder.Pin())
+		{
+			TArray<TWeakObjectPtr<UObject>> ObjectsBeingCustomized;
+			Builder->GetObjectsBeingCustomized(ObjectsBeingCustomized);
+			return ObjectsBeingCustomized.Contains(InObject);
+		}
+		return false;
+	}
+
+	TSharedPtr<SCompositePassTreePanel> GetPassPanelWidget() const { return PassPanel.Pin(); }
+
+private:
+	void OnLayoutSizeChanged()
+	{
+		if (TSharedPtr<IDetailLayoutBuilder> Builder = CachedDetailBuilder.Pin())
+		{
+			Builder->GetPropertyUtilities()->RequestRefresh();
+		}
+	}
+
+	TWeakPtr<IDetailLayoutBuilder> CachedDetailBuilder;
+	TWeakPtr<SCompositePassTreePanel> PassPanel;
+};
diff --git a/Engine/Plugins/Compositing/Composite/Source/CompositeEditor/Private/Customizations/CompositeLayerPlanarReflectionCustomization.cpp b/Engine/Plugins/Compositing/Composite/Source/CompositeEditor/Private/Customizations/CompositeLayerPlanarReflectionCustomization.cpp
new file mode 100644
--- /dev/null
+++ b/Engine/Plugins/Compositing/Composite/Source/CompositeEditor/Private/Customizations/CompositeLayerPlanarReflectionCustomization.cpp
@@ -0,0 +1,148 @@
+// Copyright Epic Games, Inc. All Rights Reserved.
+
+#include "CompositeLayerPlanarReflectionCustomization.h"
+
+#include "CompositeRenderTargetResolutionCustomization.h"
+#include "DetailCategoryBuilder.h"
+#include "DetailLayoutBuilder.h"
+#include "DetailWidgetRow.h"
+#include "IDetailGroup.h"
+#include "IPropertyUtilities.h"
+#include "Layers/CompositeLayerPlanarReflection.h"
+#include "UI/CompositeLayerPassListOwner.h"
+#include "UI/SCompositeActorPickerTable.h"
+#include "UI/SCompositePassTreePanel.h"
+#include "Widgets/Text/STextBlock.h"
+
+#define LOCTEXT_NAMESPACE "FCompositeLayerPlanarReflectionCustomization"
+
+TSharedRef<IDetailCustomization> FCompositeLayerPlanarReflectionCustomization::MakeInstance()
+{
+	return MakeShared<FCompositeLayerPlanarReflectionCustomization>();
+}
+
+void FCompositeLayerPlanarReflectionCustomization::CustomizeDetails(const TSharedPtr<IDetailLayoutBuilder>& DetailBuilder)
+{
+	CachedDetailBuilder = DetailBuilder;
+	IDetailCustomization::CustomizeDetails(DetailBuilder);
+}
+
+void FCompositeLayerPlanarReflectionCustomization::CustomizeDetails(IDetailLayoutBuilder& DetailLayout)
+{
+	DetailLayout.RegisterInstancedCustomPropertyTypeLayout(NAME_IntPoint, FOnGetPropertyTypeCustomizationInstance::CreateLambda([]
+		{
+			return MakeShared<FResolutionTypeCustomization>();
+		}),
+		MakeShared<FResolutionPropertyIdentifier>());
+
+	static TArray<FName> HiddenProperties =
+	{
+		GET_MEMBER_NAME_CHECKED(UCompositeLayerPlanarReflection, Actors),
+		GET_MEMBER_NAME_CHECKED(UCompositeLayerPlanarReflection, LayerPasses)
+	};
+
+	for (const FName& PropertyName : HiddenProperties)
+	{
+		DetailLayout.HideProperty(PropertyName);
+	}
+
+	TArray<TWeakObjectPtr<UCompositeLayerPlanarReflection>> Objects = DetailLayout.GetObjectsOfTypeBeingCustomized<UCompositeLayerPlanarReflection>();
+
+	IDetailCategoryBuilder& LayerCategory = DetailLayout.EditCategory("Composite", UCompositeLayerPlanarReflection::StaticClass()->GetDisplayNameText());
+
+	// Add all default simple properties ahead of the custom group
+	TArray<TSharedRef<IPropertyHandle>> LayerProperties;
+	LayerCategory.GetDefaultProperties(LayerProperties, /* bSimpleProperties */ true, /* bAdvancedProperties */ false);
+	for (const TSharedRef<IPropertyHandle>& Property : LayerProperties)
+	{
+		if (!HiddenProperties.Contains(Property->GetProperty()->GetNameCPP()))
+		{
+			LayerCategory.AddProperty(Property);
+		}
+	}
+
+	if (Objects.Num() == 1 && Objects[0].IsValid())
+	{
+		UCompositeLayerPlanarReflection* PlanarReflection = Cast<UCompositeLayerPlanarReflection>(Objects[0].Get());
+
+		IDetailGroup& PassesGroup = LayerCategory.AddGroup("Passes", LOCTEXT("PassesGroupName", "Passes"), false, true);
+
+		PassesGroup.AddWidgetRow()
+		.WholeRowContent()
+		[
+			SAssignNew(PassPanel, SCompositePassTreePanel, MakeShared<TCompositeLayerPassListOwner<UCompositeLayerPlanarReflection>>(PlanarReflection))
+			.OnLayoutSizeChanged(this, &FCompositeLayerPlanarReflectionCustomization::OnLayoutSizeChanged)
+		];
+
+		IDetailGroup& ActorListGroup = LayerCategory.AddGroup("ActorContent", LOCTEXT("ActorListGroupName", "Planar Reflection Content"), false, true);
+		FCompositeActorPickerListRef ActorListRef(PlanarReflection, GET_MEMBER_NAME_CHECKED(UCompositeLayerPlanarReflection, Actors), &PlanarReflection->Actors);
+
+		ActorListGroup.AddWidgetRow()
+		.WholeRowContent()
+		[
+			SNew(SCompositeActorPickerTable, ActorListRef)
+			.OnLayoutSizeChanged(this, &FCompositeLayerPlanarReflectionCustomization::OnLayoutSizeChanged)
+		];
+	}
+	else
+	{
+		// Can't display pass panel if multiple layers are selected, so simply put a "Multiple Values" entry in the property list
+		LayerCategory.AddCustomRow(LOCTEXT("PassesGroupName", "Passes"))
+		.NameContent()
+		[
+			SNew(STextBlock)
+			.Text(LOCTEXT("PassesGroupName", "Passes"))
+			.Font(IDetailLayoutBuilder::GetDetailFont())
+		]
+		.ValueContent()
+		[
+			SNew(STextBlock)
+			.Text(LOCTEXT("MultipleValues", "Multiple Values"))
+			.Font(IDetailLayoutBuilder::GetDetailFont())
+		];
+
+		// Can't display actor list if multiple layers are selected, so simply put a "Multiple Values" entry in the property list
+		LayerCategory.AddCustomRow(LOCTEXT("ActorListGroupName", "Planar Reflection Content"))
+		.NameContent()
+		[
+			SNew(STextBlock)
+			.Text(LOCTEXT("ActorListGroupName", "Planar Reflection Content"))
+			.Font(IDetailLayoutBuilder::GetDetailFont())
+		]
+		.ValueContent()
+		[
+			SNew(STextBlock)
+			.Text(LOCTEXT("MultipleValues", "Multiple Values"))
+			.Font(IDetailLayoutBuilder::GetDetailFont())
+		];
+	}
+}
+
+bool FCompositeLayerPlanarReflectionCustomization::IsCustomizingObject(UObject* InObject) const
+{
+	if (TSharedPtr<IDetailLayoutBuilder> Builder = CachedDetailBuilder.Pin())
+	{
+		TArray<TWeakObjectPtr<UObject>> ObjectsBeingCustomized;
+		Builder->GetObjectsBeingCustomized(ObjectsBeingCustomized);
+
+		return ObjectsBeingCustomized.Contains(InObject);
+	}
+
+	return false;
+}
+
+TSharedPtr<SCompositePassTreePanel> FCompositeLayerPlanarReflectionCustomization::GetPassPanelWidget() const
+{
+	return PassPanel.Pin();
+}
+
+void FCompositeLayerPlanarReflectionCustomization::OnLayoutSizeChanged()
+{
+	TSharedPtr<IDetailLayoutBuilder> PinnedDetailBuilder = CachedDetailBuilder.Pin();
+	if (PinnedDetailBuilder.IsValid())
+	{
+		PinnedDetailBuilder->GetPropertyUtilities()->RequestRefresh();
+	}
+}
+
+#undef LOCTEXT_NAMESPACE
diff --git a/Engine/Plugins/Compositing/Composite/Source/CompositeEditor/Private/Customizations/CompositeLayerPlanarReflectionCustomization.h b/Engine/Plugins/Compositing/Composite/Source/CompositeEditor/Private/Customizations/CompositeLayerPlanarReflectionCustomization.h
new file mode 100644
--- /dev/null
+++ b/Engine/Plugins/Compositing/Composite/Source/CompositeEditor/Private/Customizations/CompositeLayerPlanarReflectionCustomization.h
@@ -0,0 +1,36 @@
+// Copyright Epic Games, Inc. All Rights Reserved.
+
+#pragma once
+
+#include "IDetailCustomization.h"
+
+class SCompositePassTreePanel;
+
+/**
+ * Customization for the composite planar reflection layer, primarily for displaying custom widgets for picking actors and managing passes
+ */
+class FCompositeLayerPlanarReflectionCustomization : public IDetailCustomization
+{
+public:
+	/** Makes a new instance of the details customization */
+	static TSharedRef<IDetailCustomization> MakeInstance();
+
+	// IDetailCustomization interface
+	virtual void CustomizeDetails(const TSharedPtr<IDetailLayoutBuilder>& DetailBuilder) override;
+	virtual void CustomizeDetails(IDetailLayoutBuilder& DetailLayout) override;
+	// End of IDetailCustomization interface
+
+	/** Gets whether the specified object is one of the objects being customized by this details panel customization */
+	bool IsCustomizingObject(UObject* InObject) const;
+
+	/** Gets the pass panel widget the details panel is displaying */
+	TSharedPtr<SCompositePassTreePanel> GetPassPanelWidget() const;
+
+private:
+	/** Raised when the custom widgets' layout size may have changed */
+	void OnLayoutSizeChanged();
+
+private:
+	TWeakPtr<IDetailLayoutBuilder> CachedDetailBuilder;
+	TWeakPtr<SCompositePassTreePanel> PassPanel;
+};
diff --git a/Engine/Plugins/Compositing/Composite/Source/CompositeEditor/Private/Customizations/CompositeLayerPlateCustomization.cpp b/Engine/Plugins/Compositing/Composite/Source/CompositeEditor/Private/Customizations/CompositeLayerPlateCustomization.cpp
--- a/Engine/Plugins/Compositing/Composite/Source/CompositeEditor/Private/Customizations/CompositeLayerPlateCustomization.cpp
+++ b/Engine/Plugins/Compositing/Composite/Source/CompositeEditor/Private/Customizations/CompositeLayerPlateCustomization.cpp
@@ -2,9 +2,7 @@
 
 #include "CompositeLayerPlateCustomization.h"
 
-#include "ActorTreeItem.h"
 #include "CompositeActor.h"
-#include "CompositeCoreSettings.h"
 #include "CompositeEditorStyle.h"
 #include "CompositeMeshActor.h"
 #include "DetailCategoryBuilder.h"
@@ -18,11 +16,9 @@
 #include "MediaSource.h"
 #include "MediaTexture.h"
 #include "PropertyCustomizationHelpers.h"
-#include "SceneOutlinerFilters.h"
 #include "ScopedTransaction.h"
 #include "UnrealEdGlobals.h"
 #include "Camera/CameraComponent.h"
-#include "Components/PrimitiveComponent.h"
 #include "Editor/UnrealEdEngine.h"
 #include "Framework/MultiBox/MultiBoxBuilder.h"
 #include "Layers/CompositeLayerPlate.h"
@@ -112,26 +108,8 @@ void FCompositeLayerPlateCustomization::CustomizeDetails(IDetailLayoutBuilder& D
 			SAssignNew(CompositeMeshPicker, SCompositeActorPickerTable, ActorListRef)
 			.SceneOutlinerFilters_Lambda([]()
 			{
-				TSharedPtr<FSceneOutlinerFilters> Filters = MakeShared<FSceneOutlinerFilters>();
-
-				Filters->AddFilterPredicate<FActorTreeItem>(FActorTreeItem::FFilterPredicate::CreateLambda([](const AActor* InActor)
-				{
-					const UCompositeCorePluginSettings* Settings = GetDefault<UCompositeCorePluginSettings>();
-
-					TInlineComponentArray<UPrimitiveComponent*> PrimComponents(InActor);
-					
-					for (const UPrimitiveComponent* PrimitiveComponent : PrimComponents)
-					{
-						if (Settings->IsAllowedPrimitiveClass(PrimitiveComponent))
-						{
-							return true;
-						}
-					}
-
-					return false;
-				}));
-
-				return Filters;
+				constexpr bool bExcludeCompositeMeshActors = false;
+				return SCompositeActorPickerTable::MakeDefaultSceneOutlinerFilters(bExcludeCompositeMeshActors);
 			})
 			.OnExtendAddMenu(this, &FCompositeLayerPlateCustomization::OnExtendCompositeMeshAddMenu)
 			.OnLayoutSizeChanged(this, &FCompositeLayerPlateCustomization::OnLayoutSizeChanged)
diff --git a/Engine/Plugins/Compositing/Composite/Source/CompositeEditor/Private/Customizations/CompositeLayerSceneCaptureCustomization.cpp b/Engine/Plugins/Compositing/Composite/Source/CompositeEditor/Private/Customizations/CompositeLayerSceneCaptureCustomization.cpp
--- a/Engine/Plugins/Compositing/Composite/Source/CompositeEditor/Private/Customizations/CompositeLayerSceneCaptureCustomization.cpp
+++ b/Engine/Plugins/Compositing/Composite/Source/CompositeEditor/Private/Customizations/CompositeLayerSceneCaptureCustomization.cpp
@@ -9,7 +9,7 @@
 #include "IDetailGroup.h"
 #include "IPropertyUtilities.h"
 #include "Layers/CompositeLayerSceneCapture.h"
-#include "UI/CompositeSceneCapturePassListOwner.h"
+#include "UI/CompositeLayerPassListOwner.h"
 #include "UI/SCompositeActorPickerTable.h"
 #include "UI/SCompositePassTreePanel.h"
 #include "Widgets/Text/STextBlock.h"
@@ -70,7 +70,7 @@ void FCompositeLayerSceneCaptureCustomization::CustomizeDetails(IDetailLayoutBui
 		PassesGroup.AddWidgetRow()
 		.WholeRowContent()
 		[
-			SAssignNew(PassPanel, SCompositePassTreePanel, MakeShared<FCompositeSceneCapturePassListOwner>(SceneCapture))
+			SAssignNew(PassPanel, SCompositePassTreePanel, MakeShared<TCompositeLayerPassListOwner<UCompositeLayerSceneCapture>>(SceneCapture))
 			.OnLayoutSizeChanged(this, &FCompositeLayerSceneCaptureCustomization::OnLayoutSizeChanged)
 		];
 		
@@ -81,6 +81,11 @@ void FCompositeLayerSceneCaptureCustomization::CustomizeDetails(IDetailLayoutBui
 		.WholeRowContent()
 		[
 			SNew(SCompositeActorPickerTable, ActorListRef)
+			.SceneOutlinerFilters_Lambda([]()
+			{
+				constexpr bool bExcludeCompositeMeshActors = false;
+				return SCompositeActorPickerTable::MakeDefaultSceneOutlinerFilters(bExcludeCompositeMeshActors);
+			})
 			.OnLayoutSizeChanged(this, &FCompositeLayerSceneCaptureCustomization::OnLayoutSizeChanged)
 		];
 	}
diff --git a/Engine/Plugins/Compositing/Composite/Source/CompositeEditor/Private/Customizations/CompositeLayerShadowReflectionCustomization.cpp b/Engine/Plugins/Compositing/Composite/Source/CompositeEditor/Private/Customizations/CompositeLayerShadowReflectionCustomization.cpp
--- a/Engine/Plugins/Compositing/Composite/Source/CompositeEditor/Private/Customizations/CompositeLayerShadowReflectionCustomization.cpp
+++ b/Engine/Plugins/Compositing/Composite/Source/CompositeEditor/Private/Customizations/CompositeLayerShadowReflectionCustomization.cpp
@@ -3,14 +3,11 @@
 #include "CompositeLayerShadowReflectionCustomization.h"
 
 #include "CompositeRenderTargetResolutionCustomization.h"
-#include "ActorTreeItem.h"
-#include "CompositeMeshActor.h"
 #include "DetailCategoryBuilder.h"
 #include "DetailLayoutBuilder.h"
 #include "DetailWidgetRow.h"
 #include "IDetailGroup.h"
 #include "IPropertyUtilities.h"
-#include "SceneOutlinerFilters.h"
 #include "Layers/CompositeLayerShadowReflection.h"
 #include "UI/SCompositeActorPickerTable.h"
 #include "Widgets/Text/STextBlock.h"
@@ -65,17 +62,6 @@ void FCompositeLayerShadowReflectionCustomization::CustomizeDetails(IDetailLayou
 		.WholeRowContent()
 		[
 			SNew(SCompositeActorPickerTable, ActorListRef)
-			.SceneOutlinerFilters_Lambda([]()
-			{
-				TSharedPtr<FSceneOutlinerFilters> Filters = MakeShared<FSceneOutlinerFilters>();
-
-				Filters->AddFilterPredicate<FActorTreeItem>(FActorTreeItem::FFilterPredicate::CreateLambda([](const AActor* InActor)
-				{
-					return !InActor->IsA<ACompositeMeshActor>();
-				}));
-				
-				return Filters;
-			})
 			.OnLayoutSizeChanged(this, &FCompositeLayerShadowReflectionCustomization::OnLayoutSizeChanged)
 		];
 	}
diff --git a/Engine/Plugins/Compositing/Composite/Source/CompositeEditor/Private/UI/CompositeLayerPassListOwner.h b/Engine/Plugins/Compositing/Composite/Source/CompositeEditor/Private/UI/CompositeLayerPassListOwner.h
new file mode 100644
--- /dev/null
+++ b/Engine/Plugins/Compositing/Composite/Source/CompositeEditor/Private/UI/CompositeLayerPassListOwner.h
@@ -0,0 +1,175 @@
+// Copyright Epic Games, Inc. All Rights Reserved.
+
+#pragma once
+
+#include "PropertyEditorUtils.h"
+#include "SCompositePassTree.h"
+#include "UObject/Package.h"
+
+/**
+ * Generic ICompositePassListOwner for any layer type with a single public LayerPasses array.
+ * Covers UCompositeLayerSceneCapture, UCompositeLayerPlanarReflection, UCompositeLayerMainRender, UCompositeLayerProcessing.
+ * (UCompositeLayerPlate uses its own owner because it has multiple pass groups.)
+ */
+template<typename LayerClass>
+class TCompositeLayerPassListOwner : public ICompositePassListOwner
+{
+public:
+	TCompositeLayerPassListOwner(const TWeakObjectPtr<LayerClass>& InLayer)
+		: Layer(InLayer)
+	{ }
+
+	virtual bool IsObjectValid() override { return Layer.IsValid(); }
+	virtual TStrongObjectPtr<UObject> GetObject() override { return Layer.Pin(); }
+
+	virtual bool IsPassListPropertyName(const FName& InPropertyName) override
+	{
+		return InPropertyName == GET_MEMBER_NAME_CHECKED(LayerClass, LayerPasses);
+	}
+
+	virtual bool IsValidPassIndex(int32 InGroupIndex, int32 InPassIndex) override
+	{
+		if (InGroupIndex != INDEX_NONE) { return false; }
+		TStrongObjectPtr<LayerClass> Pinned = Layer.Pin();
+		return Pinned.IsValid() && Pinned->LayerPasses.IsValidIndex(InPassIndex);
+	}
+
+	virtual UCompositePassBase* GetPass(int32 InGroupIndex, int32 InPassIndex) override
+	{
+		if (InGroupIndex != INDEX_NONE) { return nullptr; }
+		TStrongObjectPtr<LayerClass> Pinned = Layer.Pin();
+		if (!Pinned.IsValid() || !Pinned->LayerPasses.IsValidIndex(InPassIndex)) { return nullptr; }
+		return Pinned->LayerPasses[InPassIndex];
+	}
+
+	virtual int32 GetNumGroups() const override { return INDEX_NONE; }
+
+	virtual TArray<TObjectPtr<UCompositePassBase>>& GetPassesForGroup(int32 InGroupIndex) override
+	{
+		TStrongObjectPtr<LayerClass> Pinned = Layer.Pin();
+		check(Pinned.IsValid());
+		return Pinned->LayerPasses;
+	}
+
+	virtual FString GetGroupFilterString(int32 InGroupIndex) override { return TEXT(""); }
+	virtual const FSlateBrush* GetGroupIcon(int32 InGroupIndex) override { return nullptr; }
+	virtual FText GetGroupDisplayName(int32 InGroupIndex) override { return FText::GetEmpty(); }
+	virtual FGroupFilterConfig GetGroupFilterConfig(int32 InGroupIndex) override { return FGroupFilterConfig(); }
+	virtual int32 GetDefaultGroupForNewPass(const UClass* InPassClass) const override { return INDEX_NONE; }
+
+	virtual bool CanAddPass(const UClass* InPassClass, int32 InGroupIndex, int32 InPassIndex) const override
+	{
+		if (InGroupIndex != INDEX_NONE) { return false; }
+		TStrongObjectPtr<LayerClass> Pinned = Layer.Pin();
+		if (!Pinned.IsValid()) { return false; }
+
+		FProperty* Prop = FindFProperty<FProperty>(LayerClass::StaticClass(), GET_MEMBER_NAME_CHECKED(LayerClass, LayerPasses));
+		TArray<const UClass*> Allowed, Disallowed;
+		PropertyEditorUtils::GetAllowedAndDisallowedClasses({ Pinned.Get() }, *Prop, Allowed, Disallowed, false);
+		return (Allowed.IsEmpty() || Allowed.Contains(InPassClass)) && !Disallowed.Contains(InPassClass);
+	}
+
+	virtual int32 AddPass(const UClass* InPassClass, int32 InGroupIndex, int32 InPassIndex) override
+	{
+		if (InGroupIndex != INDEX_NONE) { return INDEX_NONE; }
+		TStrongObjectPtr<LayerClass> Pinned = Layer.Pin();
+		if (!Pinned.IsValid()) { return INDEX_NONE; }
+
+		const int32 ClampedIndex = FMath::Clamp(InPassIndex, 0, Pinned->LayerPasses.Num());
+
+		Pinned->Modify();
+		UCompositePassBase* NewPass = NewObject<UCompositePassBase>(Pinned.Get(), InPassClass, NAME_None, RF_Transactional);
+
+		FProperty* Prop = GetPassListProperty();
+		Pinned->PreEditChange(Prop);
+		int32 NewIndex = Pinned->LayerPasses.Insert(NewPass, ClampedIndex);
+		FPropertyChangedEvent Event(Prop, EPropertyChangeType::ArrayAdd);
+		Pinned->PostEditChangeProperty(Event);
+		return NewIndex;
+	}
+
+	virtual TArray<int32> CopyPasses(const TArray<UCompositePassBase*>& InPassesToCopy, int32 InGroupIndex, int32 InPassIndex) override
+	{
+		TArray<int32> Out;
+		if (InGroupIndex != INDEX_NONE) { return Out; }
+		TStrongObjectPtr<LayerClass> Pinned = Layer.Pin();
+		if (!Pinned.IsValid()) { return Out; }
+
+		Pinned->Modify();
+
+		FProperty* Prop = GetPassListProperty();
+		Pinned->PreEditChange(Prop);
+
+		int32 Dest = FMath::Clamp(InPassIndex, 0, Pinned->LayerPasses.Num());
+		for (UCompositePassBase* Pass : InPassesToCopy)
+		{
+			UCompositePassBase* Copy = DuplicateObject(Pass, Pinned.Get());
+			Copy->SetFlags(RF_Transactional);
+			Out.Add(Pinned->LayerPasses.Insert(Copy, Dest));
+			++Dest;
+		}
+
+		FPropertyChangedEvent Event(Prop, EPropertyChangeType::ArrayAdd);
+		Pinned->PostEditChangeProperty(Event);
+		return Out;
+	}
+
+	virtual void MovePass(int32 InSourceGroupIndex, int32 InSourcePassIndex, int32 InDestGroupIndex, int32 InDestPassIndex) override
+	{
+		if (InSourceGroupIndex != INDEX_NONE || InDestGroupIndex != INDEX_NONE) { return; }
+		TStrongObjectPtr<LayerClass> Pinned = Layer.Pin();
+		if (!Pinned.IsValid()) { return; }
+		if (!Pinned->LayerPasses.IsValidIndex(InSourcePassIndex)) { return; }
+
+		Pinned->Modify();
+		FProperty* Prop = GetPassListProperty();
+		UCompositePassBase* Pass = Pinned->LayerPasses[InSourcePassIndex];
+
+		// InDestPassIndex is a post-removal index: the caller (SCompositePassTree) adjusts for the shift caused by RemoveAt.
+		Pinned->PreEditChange(Prop);
+		Pinned->LayerPasses.RemoveAt(InSourcePassIndex);
+		Pinned->LayerPasses.Insert(Pass, FMath::Clamp(InDestPassIndex, 0, Pinned->LayerPasses.Num()));
+
+		FPropertyChangedEvent Event(Prop, EPropertyChangeType::ArrayMove);
+		Pinned->PostEditChangeProperty(Event);
+	}
+
+	virtual void RemovePasses(int32 InGroupIndex, const TArray<int32>& InPassIndices) override
+	{
+		if (InGroupIndex != INDEX_NONE) { return; }
+		TStrongObjectPtr<LayerClass> Pinned = Layer.Pin();
+		if (!Pinned.IsValid()) { return; }
+
+		Pinned->Modify();
+
+		TArray<int32> Sorted = InPassIndices;
+		Sorted.Sort();
+
+		FProperty* Prop = GetPassListProperty();
+		Pinned->PreEditChange(Prop);
+
+		for (int32 i = Sorted.Num() - 1; i >= 0; --i)
+		{
+			if (!Pinned->LayerPasses.IsValidIndex(Sorted[i])) { continue; }
+			TObjectPtr<UCompositePassBase>& PassToRemove = Pinned->LayerPasses[Sorted[i]];
+			if (IsValid(PassToRemove))
+			{
+				PassToRemove->Modify();
+				PassToRemove->Rename(nullptr, GetTransientPackage(), REN_DontCreateRedirectors);
+			}
+			Pinned->LayerPasses.RemoveAt(Sorted[i], EAllowShrinking::No);
+		}
+
+		Pinned->LayerPasses.Shrink();
+		FPropertyChangedEvent Event(Prop, EPropertyChangeType::ArrayRemove);
+		Pinned->PostEditChangeProperty(Event);
+	}
+
+private:
+	FProperty* GetPassListProperty() const
+	{
+		return FindFProperty<FProperty>(LayerClass::StaticClass(), GET_MEMBER_NAME_CHECKED(LayerClass, LayerPasses));
+	}
+
+	TWeakObjectPtr<LayerClass> Layer;
+};
diff --git a/Engine/Plugins/Compositing/Composite/Source/CompositeEditor/Private/UI/CompositePlatePassListOwner.cpp b/Engine/Plugins/Compositing/Composite/Source/CompositeEditor/Private/UI/CompositePlatePassListOwner.cpp
--- a/Engine/Plugins/Compositing/Composite/Source/CompositeEditor/Private/UI/CompositePlatePassListOwner.cpp
+++ b/Engine/Plugins/Compositing/Composite/Source/CompositeEditor/Private/UI/CompositePlatePassListOwner.cpp
@@ -278,7 +278,7 @@ void FCompositePlatePassListOwner::RemovePasses(int32 InGroupIndex, const TArray
 
 	PassList.Shrink();
 			
-	FPropertyChangedEvent PropertyChangedEvent(PassListProperty, EPropertyChangeType::ArrayAdd);
+	FPropertyChangedEvent PropertyChangedEvent(PassListProperty, EPropertyChangeType::ArrayRemove);
 	PinnedPlate->PostEditChangeProperty(PropertyChangedEvent);
 }
 
diff --git a/Engine/Plugins/Compositing/Composite/Source/CompositeEditor/Private/UI/CompositeSceneCapturePassListOwner.cpp b/Engine/Plugins/Compositing/Composite/Source/CompositeEditor/Private/UI/CompositeSceneCapturePassListOwner.cpp
deleted file mode 100644
diff --git a/Engine/Plugins/Compositing/Composite/Source/CompositeEditor/Private/UI/CompositeSceneCapturePassListOwner.h b/Engine/Plugins/Compositing/Composite/Source/CompositeEditor/Private/UI/CompositeSceneCapturePassListOwner.h
deleted file mode 100644
diff --git a/Engine/Plugins/Compositing/Composite/Source/CompositeEditor/Private/UI/SCompositeActorPickerTable.cpp b/Engine/Plugins/Compositing/Composite/Source/CompositeEditor/Private/UI/SCompositeActorPickerTable.cpp
--- a/Engine/Plugins/Compositing/Composite/Source/CompositeEditor/Private/UI/SCompositeActorPickerTable.cpp
+++ b/Engine/Plugins/Compositing/Composite/Source/CompositeEditor/Private/UI/SCompositeActorPickerTable.cpp
@@ -2,7 +2,11 @@
 
 #include "SCompositeActorPickerTable.h"
 
+#include "ActorTreeItem.h"
+#include "CompositeCoreSettings.h"
+#include "CompositeMeshActor.h"
 #include "Components/PrimitiveComponent.h"
+#include "SceneOutlinerFilters.h"
 #include "CompositeEditorCommands.h"
 #include "Editor.h"
 #include "EditorActorFolders.h"
@@ -153,9 +157,15 @@ private:
 		
 		MenuBuilder.BeginSection("SceneOutliner", LOCTEXT("SceneOutlinerSectionLabel", "Browse"));
 		{
+			TSharedPtr<FSceneOutlinerFilters> Filters = SceneOutlinerFilters.Get(nullptr);
+			if (!Filters)
+			{
+				Filters = SCompositeActorPickerTable::MakeDefaultSceneOutlinerFilters();
+			}
+
 			TSharedRef<SWidget> SceneOutliner =
 				SNew(SCompositeActorPickerSceneOutliner, ActorListRef)
-				.SceneOutlinerFilters(SceneOutlinerFilters.Get(nullptr))
+				.SceneOutlinerFilters(Filters)
 				.OnActorListChanged(OnActorListChanged);
 			
 			MenuBuilder.AddWidget(SceneOutliner, FText::GetEmpty(), true);
@@ -389,6 +399,39 @@ FText FCompositeActorPickerListRef::GetToolTipText(bool bShortTooltip) const
 	return ActorListProperty ? ActorListProperty->GetToolTipText(bShortTooltip) : FText::GetEmpty();
 }
 
+TSharedPtr<FSceneOutlinerFilters> SCompositeActorPickerTable::MakeDefaultSceneOutlinerFilters(bool bExcludeCompositeMeshActors)
+{
+	TSharedPtr<FSceneOutlinerFilters> Filters = MakeShared<FSceneOutlinerFilters>();
+
+	Filters->AddFilterPredicate<FActorTreeItem>(FActorTreeItem::FFilterPredicate::CreateLambda([bExcludeCompositeMeshActors](const AActor* InActor)
+	{
+		if (!InActor)
+		{
+			return false;
+		}
+
+		if (bExcludeCompositeMeshActors && InActor->IsA<ACompositeMeshActor>())
+		{
+			return false;
+		}
+
+		const UCompositeCorePluginSettings* Settings = GetDefault<UCompositeCorePluginSettings>();
+		TInlineComponentArray<UPrimitiveComponent*> PrimComponents(InActor);
+
+		for (const UPrimitiveComponent* PrimitiveComponent : PrimComponents)
+		{
+			if (Settings->IsAllowedPrimitiveClass(PrimitiveComponent))
+			{
+				return true;
+			}
+		}
+
+		return false;
+	}));
+
+	return Filters;
+}
+
 void SCompositeActorPickerTable::Construct(const FArguments& InArgs, const FCompositeActorPickerListRef& InActorListRef)
 {
 	ActorListRef = InActorListRef;
diff --git a/Engine/Plugins/Compositing/Composite/Source/CompositeEditor/Private/UI/SCompositeActorPickerTable.h b/Engine/Plugins/Compositing/Composite/Source/CompositeEditor/Private/UI/SCompositeActorPickerTable.h
--- a/Engine/Plugins/Compositing/Composite/Source/CompositeEditor/Private/UI/SCompositeActorPickerTable.h
+++ b/Engine/Plugins/Compositing/Composite/Source/CompositeEditor/Private/UI/SCompositeActorPickerTable.h
@@ -68,8 +68,11 @@ private:
 
 public:
 	DECLARE_DELEGATE_OneParam(FOnExtendAddMenu, FMenuBuilder&);
-	
+
 public:
+	/** Returns default scene outliner filters that check IsAllowedPrimitiveClass and optionally exclude ACompositeMeshActor. */
+	static TSharedPtr<FSceneOutlinerFilters> MakeDefaultSceneOutlinerFilters(bool bExcludeCompositeMeshActors = true);
+
 	SLATE_BEGIN_ARGS(SCompositeActorPickerTable) : _ShowApplyMaterialSection(false) { }
 		SLATE_ATTRIBUTE(TSharedPtr<FSceneOutlinerFilters>, SceneOutlinerFilters)
 		SLATE_EVENT(FOnExtendAddMenu, OnExtendAddMenu)
