Composure: Added analytics for adding composite actors, layers or passes.

* Editor.Usage.Composite.ActorAdded
* Editor.Usage.Composite.LayerAdded
* Editor.Usage.Composite.PassAdded

#Synced-CL 52419094

diff --git a/Engine/Plugins/Compositing/Composite/Source/CompositeEditor/Private/ActorFactories/CompositeActorFactory.cpp b/Engine/Plugins/Compositing/Composite/Source/CompositeEditor/Private/ActorFactories/CompositeActorFactory.cpp
--- a/Engine/Plugins/Compositing/Composite/Source/CompositeEditor/Private/ActorFactories/CompositeActorFactory.cpp
+++ b/Engine/Plugins/Compositing/Composite/Source/CompositeEditor/Private/ActorFactories/CompositeActorFactory.cpp
@@ -3,7 +3,7 @@
 #include "CompositeActorFactory.h"
 
 #include "CompositeActor.h"
-#include "EngineAnalytics.h"
+#include "CompositeAnalytics.h"
 #include "Layers/CompositeLayerMainRender.h"
 #include "Layers/CompositeLayerShadowReflection.h"
 #include "Layers/CompositeLayerPlate.h"
@@ -47,10 +47,7 @@ void UCompositeActorFactory::PostSpawnActor(UObject* Asset, AActor* NewActor)
 		FCoreUObjectDelegates::OnObjectPropertyChanged.Broadcast(NewActor, LayersChangedEvent);
 
 		// Record analytics
-		if (FEngineAnalytics::IsAvailable())
-		{
-			FEngineAnalytics::GetProvider().RecordEvent(TEXT("Editor.Usage.Composite.PostSpawnActor"));
-		}
+		Composite::Analytics::RecordActorAdded(*CompositeActor);
 	}
 }
 
diff --git a/Engine/Plugins/Compositing/Composite/Source/CompositeEditor/Private/ActorFactories/CompositeDepthMeshActorFactory.cpp b/Engine/Plugins/Compositing/Composite/Source/CompositeEditor/Private/ActorFactories/CompositeDepthMeshActorFactory.cpp
new file mode 100644
--- /dev/null
+++ b/Engine/Plugins/Compositing/Composite/Source/CompositeEditor/Private/ActorFactories/CompositeDepthMeshActorFactory.cpp
@@ -0,0 +1,31 @@
+// Copyright Epic Games, Inc. All Rights Reserved.
+
+#include "CompositeDepthMeshActorFactory.h"
+
+#include "CompositeAnalytics.h"
+#include "CompositeDepthMeshActor.h"
+
+#include UE_INLINE_GENERATED_CPP_BY_NAME(CompositeDepthMeshActorFactory)
+
+#define LOCTEXT_NAMESPACE "CompositeDepthMeshActorFactory"
+
+UCompositeDepthMeshActorFactory::UCompositeDepthMeshActorFactory(const FObjectInitializer& ObjInit)
+	: Super(ObjInit)
+{
+	DisplayName = LOCTEXT("DisplayName", "Composite Depth Mesh Actor");
+	NewActorClass = ACompositeDepthMeshActor::StaticClass();
+}
+
+void UCompositeDepthMeshActorFactory::PostSpawnActor(UObject* Asset, AActor* NewActor)
+{
+	Super::PostSpawnActor(Asset, NewActor);
+
+	ACompositeDepthMeshActor* DepthMeshActor = Cast<ACompositeDepthMeshActor>(NewActor);
+
+	if (IsValid(DepthMeshActor) && !DepthMeshActor->IsTemplate() && !DepthMeshActor->HasAnyFlags(RF_Transient))
+	{
+		Composite::Analytics::RecordActorAdded(*DepthMeshActor);
+	}
+}
+
+#undef LOCTEXT_NAMESPACE
diff --git a/Engine/Plugins/Compositing/Composite/Source/CompositeEditor/Private/ActorFactories/CompositeDepthMeshActorFactory.h b/Engine/Plugins/Compositing/Composite/Source/CompositeEditor/Private/ActorFactories/CompositeDepthMeshActorFactory.h
new file mode 100644
--- /dev/null
+++ b/Engine/Plugins/Compositing/Composite/Source/CompositeEditor/Private/ActorFactories/CompositeDepthMeshActorFactory.h
@@ -0,0 +1,23 @@
+// Copyright Epic Games, Inc. All Rights Reserved.
+
+#pragma once
+
+#include "UObject/ObjectMacros.h"
+#include "ActorFactories/ActorFactory.h"
+
+#include "CompositeDepthMeshActorFactory.generated.h"
+
+UCLASS()
+class UCompositeDepthMeshActorFactory : public UActorFactory
+{
+	GENERATED_BODY()
+
+public:
+	/** Constructor */
+	UCompositeDepthMeshActorFactory(const FObjectInitializer& ObjInit);
+
+public:
+	//~ Begin UActorFactory interface
+	virtual void PostSpawnActor(UObject* Asset, AActor* NewActor) override;
+	//~ End UActorFactory interface
+};
diff --git a/Engine/Plugins/Compositing/Composite/Source/CompositeEditor/Private/CompositeAnalytics.cpp b/Engine/Plugins/Compositing/Composite/Source/CompositeEditor/Private/CompositeAnalytics.cpp
new file mode 100644
--- /dev/null
+++ b/Engine/Plugins/Compositing/Composite/Source/CompositeEditor/Private/CompositeAnalytics.cpp
@@ -0,0 +1,44 @@
+// Copyright Epic Games, Inc. All Rights Reserved.
+
+#include "CompositeAnalytics.h"
+
+#include "EngineAnalytics.h"
+#include "GameFramework/Actor.h"
+#include "Layers/CompositeLayerBase.h"
+#include "Passes/CompositePassBase.h"
+
+void Composite::Analytics::RecordActorAdded(const AActor& Actor)
+{
+	if (!FEngineAnalytics::IsAvailable())
+	{
+		return;
+	}
+
+	TArray<FAnalyticsEventAttribute> Attrs;
+	Attrs.Add(FAnalyticsEventAttribute(TEXT("ActorClass"), Actor.GetClass()->GetName()));
+	FEngineAnalytics::GetProvider().RecordEvent(TEXT("Editor.Usage.Composite.ActorAdded"), Attrs);
+}
+
+void Composite::Analytics::RecordLayerAdded(const UCompositeLayerBase& Layer)
+{
+	if (!FEngineAnalytics::IsAvailable())
+	{
+		return;
+	}
+
+	TArray<FAnalyticsEventAttribute> Attrs;
+	Attrs.Add(FAnalyticsEventAttribute(TEXT("LayerClass"), Layer.GetClass()->GetName()));
+	FEngineAnalytics::GetProvider().RecordEvent(TEXT("Editor.Usage.Composite.LayerAdded"), Attrs);
+}
+
+void Composite::Analytics::RecordPassAdded(const UCompositePassBase& Pass)
+{
+	if (!FEngineAnalytics::IsAvailable())
+	{
+		return;
+	}
+
+	TArray<FAnalyticsEventAttribute> Attrs;
+	Attrs.Add(FAnalyticsEventAttribute(TEXT("PassClass"), Pass.GetClass()->GetName()));
+	FEngineAnalytics::GetProvider().RecordEvent(TEXT("Editor.Usage.Composite.PassAdded"), Attrs);
+}
diff --git a/Engine/Plugins/Compositing/Composite/Source/CompositeEditor/Private/CompositeAnalytics.h b/Engine/Plugins/Compositing/Composite/Source/CompositeEditor/Private/CompositeAnalytics.h
new file mode 100644
--- /dev/null
+++ b/Engine/Plugins/Compositing/Composite/Source/CompositeEditor/Private/CompositeAnalytics.h
@@ -0,0 +1,19 @@
+// Copyright Epic Games, Inc. All Rights Reserved.
+
+#pragma once
+
+class AActor;
+class UCompositeLayerBase;
+class UCompositePassBase;
+
+namespace Composite::Analytics
+{
+	/** Records Editor.Usage.Composite.ActorAdded with the actor's class name. */
+	void RecordActorAdded(const AActor& Actor);
+
+	/** Records Editor.Usage.Composite.LayerAdded with the layer's class name. */
+	void RecordLayerAdded(const UCompositeLayerBase& Layer);
+
+	/** Records Editor.Usage.Composite.PassAdded with the pass's class name. */
+	void RecordPassAdded(const UCompositePassBase& Pass);
+}
diff --git a/Engine/Plugins/Compositing/Composite/Source/CompositeEditor/Private/UI/SCompositePanelLayerTree.cpp b/Engine/Plugins/Compositing/Composite/Source/CompositeEditor/Private/UI/SCompositePanelLayerTree.cpp
--- a/Engine/Plugins/Compositing/Composite/Source/CompositeEditor/Private/UI/SCompositePanelLayerTree.cpp
+++ b/Engine/Plugins/Compositing/Composite/Source/CompositeEditor/Private/UI/SCompositePanelLayerTree.cpp
@@ -18,6 +18,7 @@
 #include "Filters/SBasicFilterBar.h"
 #include "Framework/Commands/GenericCommands.h"
 #include "HAL/PlatformApplicationMisc.h"
+#include "CompositeAnalytics.h"
 #include "Layers/CompositeLayerBase.h"
 #include "SLevelViewport.h"
 #include "Styling/SlateIconFinder.h"
@@ -2087,7 +2088,9 @@ void SCompositePanelLayerTree::OnLayerAdded(const UClass* InLayerClass)
 	FPropertyChangedEvent ChangeEvent(CompositeLayersProperty, EPropertyChangeType::ArrayAdd);
 	FPropertyChangedChainEvent ChainEvent(EditChain, ChangeEvent);
 	CompositeActor->PostEditChangeChainProperty(ChainEvent);
-	
+
+	Composite::Analytics::RecordLayerAdded(*NewLayer);
+
 	RefreshCompositeActorLayers(CompositeActor.Get());
 }
 
diff --git a/Engine/Plugins/Compositing/Composite/Source/CompositeEditor/Private/UI/SCompositePassTree.cpp b/Engine/Plugins/Compositing/Composite/Source/CompositeEditor/Private/UI/SCompositePassTree.cpp
--- a/Engine/Plugins/Compositing/Composite/Source/CompositeEditor/Private/UI/SCompositePassTree.cpp
+++ b/Engine/Plugins/Compositing/Composite/Source/CompositeEditor/Private/UI/SCompositePassTree.cpp
@@ -15,6 +15,7 @@
 #include "Filters/SBasicFilterBar.h"
 #include "Framework/Commands/GenericCommands.h"
 #include "HAL/PlatformApplicationMisc.h"
+#include "CompositeAnalytics.h"
 #include "Layers/CompositeLayerPlate.h"
 #include "Misc/StringOutputDevice.h"
 #include "Passes/CompositePassDistortion.h"
@@ -1701,8 +1702,10 @@ void SCompositePassTree::OnPassAdded(const UClass* InPassClass)
 	{
 		TArray<TObjectPtr<UCompositePassBase>>& Passes = PassListOwner->GetPassesForGroup(GroupIndex);
 		const FString UniqueName = UCompositePassBase::MakeUniqueDisplayName(NewPass->GetDisplayName(), MakeArrayView(Passes), NewPass);
-		
+
 		NewPass->SetDisplayName(UniqueName);
+
+		Composite::Analytics::RecordPassAdded(*NewPass);
 	}
 
 	if (Toolbar.IsValid())
