Composure: Added rendering unit tests to relevant composite passes.

#Synced-CL 52419094

diff --git a/Engine/Plugins/Compositing/Composite/Source/Composite/Private/Passes/CompositePassBlur.cpp b/Engine/Plugins/Compositing/Composite/Source/Composite/Private/Passes/CompositePassBlur.cpp
--- a/Engine/Plugins/Compositing/Composite/Source/Composite/Private/Passes/CompositePassBlur.cpp
+++ b/Engine/Plugins/Compositing/Composite/Source/Composite/Private/Passes/CompositePassBlur.cpp
@@ -30,90 +30,76 @@ class FCompositePassBlurShader : public FGlobalShader
 IMPLEMENT_GLOBAL_SHADER(FCompositePassBlurShader, "/Plugin/Composite/Private/CompositeBlur.usf", "MainPS", SF_Pixel);
 
 
-namespace UE
+namespace UE::Composite::Private
 {
-	namespace Composite
+	using namespace CompositeCore;
+
+	FPassTexture FCompositePassBlurProxy::Add(FRDGBuilder& GraphBuilder, const FSceneView& InView, const FPassInputArray& Inputs, const FPassContext& PassContext) const
 	{
-		using namespace CompositeCore;
+		RDG_EVENT_SCOPE_STAT(GraphBuilder, FCompositeBlur, "Composite.Blur");
+
+		check(ValidateInputs(Inputs));
+
+		const UE::CompositeCore::FResourceMetadata& Metadata = Inputs[0].Metadata;
+		const FScreenPassTexture& InputTexture = Inputs[0].Texture;
+		FScreenPassRenderTarget IntermediateOutput = CreateOutputRenderTarget(GraphBuilder, InView, PassContext.OutputViewRect, InputTexture.Texture->Desc, TEXT("CompositePassBlurIntermediate"));
+		FScreenPassRenderTarget Output = Inputs.OverrideOutput;
+
+		if (!Output.IsValid())
+		{
+			Output = CreateOutputRenderTarget(GraphBuilder, InView, PassContext.OutputViewRect, InputTexture.Texture->Desc, TEXT("CompositePassBlurOutput"));
+		}
+
+		FGlobalShaderMap* GlobalShaderMap = GetGlobalShaderMap(InView.GetFeatureLevel());
+
+		FCompositePassBlurShader::FPermutationDomain PermutationVector;
+		PermutationVector.Set<FCompositePassBlurShader::FAlphaOnlyDim>(bAlphaOnly);
+
+		{
+			FCompositePassBlurShader::FParameters* Parameters = GraphBuilder.AllocParameters<FCompositePassBlurShader::FParameters>();
+			Parameters->Input = GetScreenPassTextureInput(InputTexture, TStaticSamplerState<SF_Bilinear, AM_Clamp, AM_Clamp, AM_Clamp>::GetRHI());
+			Parameters->SvPositionToInputTextureUV = (
+				FScreenTransform::ChangeTextureBasisFromTo(FScreenPassTextureViewport(IntermediateOutput), FScreenTransform::ETextureBasis::TexelPosition, FScreenTransform::ETextureBasis::ViewportUV) *
+				FScreenTransform::ChangeTextureBasisFromTo(FScreenPassTextureViewport(InputTexture), FScreenTransform::ETextureBasis::ViewportUV, FScreenTransform::ETextureBasis::TextureUV));
+			Parameters->Dir = FVector2f(0, 1);
+			Parameters->Radius = Radius;
+			Parameters->bInvertedAlpha = static_cast<uint32>(Metadata.bInvertedAlpha);
+			Parameters->RenderTargets[0] = IntermediateOutput.GetRenderTargetBinding();
+
+			TShaderMapRef<FCompositePassBlurShader> PixelShader(GlobalShaderMap, PermutationVector);
+			FPixelShaderUtils::AddFullscreenPass(
+				GraphBuilder,
+				GlobalShaderMap,
+				RDG_EVENT_NAME("Composite.Blur.Vertical (%dx%d)", IntermediateOutput.ViewRect.Width(), IntermediateOutput.ViewRect.Height()),
+				PixelShader,
+				Parameters,
+				IntermediateOutput.ViewRect
+			);
+		}
 
-		class FCompositePassBlurProxy : public FCompositeCorePassProxy
 		{
-		public:
-			IMPLEMENT_COMPOSITE_PASS(FCompositePassBlurProxy);
-
-			using FCompositeCorePassProxy::FCompositeCorePassProxy;
-
-			FPassTexture Add(FRDGBuilder& GraphBuilder, const FSceneView& InView, const FPassInputArray& Inputs, const FPassContext& PassContext) const override
-			{
-				RDG_EVENT_SCOPE_STAT(GraphBuilder, FCompositeBlur, "Composite.Blur");
-
-				check(ValidateInputs(Inputs));
-
-				const UE::CompositeCore::FResourceMetadata& Metadata = Inputs[0].Metadata;
-				const FScreenPassTexture& InputTexture = Inputs[0].Texture;
-				FScreenPassRenderTarget IntermediateOutput = CreateOutputRenderTarget(GraphBuilder, InView, PassContext.OutputViewRect, InputTexture.Texture->Desc, TEXT("CompositePassBlurIntermediate"));
-				FScreenPassRenderTarget Output = Inputs.OverrideOutput;
-
-				if (!Output.IsValid())
-				{
-					Output = CreateOutputRenderTarget(GraphBuilder, InView, PassContext.OutputViewRect, InputTexture.Texture->Desc, TEXT("CompositePassBlurOutput"));
-				}
-
-				FGlobalShaderMap* GlobalShaderMap = GetGlobalShaderMap(InView.GetFeatureLevel());
-
-				FCompositePassBlurShader::FPermutationDomain PermutationVector;
-				PermutationVector.Set<FCompositePassBlurShader::FAlphaOnlyDim>(bAlphaOnly);
-
-				{
-					FCompositePassBlurShader::FParameters* Parameters = GraphBuilder.AllocParameters<FCompositePassBlurShader::FParameters>();
-					Parameters->Input = GetScreenPassTextureInput(InputTexture, TStaticSamplerState<SF_Bilinear, AM_Clamp, AM_Clamp, AM_Clamp>::GetRHI());
-					Parameters->SvPositionToInputTextureUV = (
-						FScreenTransform::ChangeTextureBasisFromTo(FScreenPassTextureViewport(IntermediateOutput), FScreenTransform::ETextureBasis::TexelPosition, FScreenTransform::ETextureBasis::ViewportUV) *
-						FScreenTransform::ChangeTextureBasisFromTo(FScreenPassTextureViewport(InputTexture), FScreenTransform::ETextureBasis::ViewportUV, FScreenTransform::ETextureBasis::TextureUV));
-					Parameters->Dir = FVector2f(0, 1);
-					Parameters->Radius = Radius;
-					Parameters->bInvertedAlpha = static_cast<uint32>(Metadata.bInvertedAlpha);
-					Parameters->RenderTargets[0] = IntermediateOutput.GetRenderTargetBinding();
-
-					TShaderMapRef<FCompositePassBlurShader> PixelShader(GlobalShaderMap, PermutationVector);
-					FPixelShaderUtils::AddFullscreenPass(
-						GraphBuilder,
-						GlobalShaderMap,
-						RDG_EVENT_NAME("Composite.Blur.Vertical (%dx%d)", IntermediateOutput.ViewRect.Width(), IntermediateOutput.ViewRect.Height()),
-						PixelShader,
-						Parameters,
-						IntermediateOutput.ViewRect
-					);
-				}
-
-				{
-					FCompositePassBlurShader::FParameters* Parameters = GraphBuilder.AllocParameters<FCompositePassBlurShader::FParameters>();
-					Parameters->Input = GetScreenPassTextureInput(IntermediateOutput, TStaticSamplerState<SF_Bilinear, AM_Clamp, AM_Clamp, AM_Clamp>::GetRHI());
-					Parameters->SvPositionToInputTextureUV = (
-						FScreenTransform::ChangeTextureBasisFromTo(FScreenPassTextureViewport(Output), FScreenTransform::ETextureBasis::TexelPosition, FScreenTransform::ETextureBasis::ViewportUV) *
-						FScreenTransform::ChangeTextureBasisFromTo(FScreenPassTextureViewport(IntermediateOutput), FScreenTransform::ETextureBasis::ViewportUV, FScreenTransform::ETextureBasis::TextureUV));
-					Parameters->Dir = FVector2f(1, 0);
-					Parameters->Radius = Radius;
-					Parameters->bInvertedAlpha = static_cast<uint32>(Metadata.bInvertedAlpha);
-					Parameters->RenderTargets[0] = Output.GetRenderTargetBinding();
-
-					TShaderMapRef<FCompositePassBlurShader> PixelShader(GlobalShaderMap, PermutationVector);
-					FPixelShaderUtils::AddFullscreenPass(
-						GraphBuilder,
-						GlobalShaderMap,
-						RDG_EVENT_NAME("Composite.Blur.Horizontal (%dx%d)", Output.ViewRect.Width(), Output.ViewRect.Height()),
-						PixelShader,
-						Parameters,
-						Output.ViewRect
-					);
-				}
-
-				return UE::CompositeCore::FPassTexture{ MoveTemp(Output), Metadata };
-			}
-
-			int32 Radius = 1;
-			bool bAlphaOnly = false;
-		};
+			FCompositePassBlurShader::FParameters* Parameters = GraphBuilder.AllocParameters<FCompositePassBlurShader::FParameters>();
+			Parameters->Input = GetScreenPassTextureInput(IntermediateOutput, TStaticSamplerState<SF_Bilinear, AM_Clamp, AM_Clamp, AM_Clamp>::GetRHI());
+			Parameters->SvPositionToInputTextureUV = (
+				FScreenTransform::ChangeTextureBasisFromTo(FScreenPassTextureViewport(Output), FScreenTransform::ETextureBasis::TexelPosition, FScreenTransform::ETextureBasis::ViewportUV) *
+				FScreenTransform::ChangeTextureBasisFromTo(FScreenPassTextureViewport(IntermediateOutput), FScreenTransform::ETextureBasis::ViewportUV, FScreenTransform::ETextureBasis::TextureUV));
+			Parameters->Dir = FVector2f(1, 0);
+			Parameters->Radius = Radius;
+			Parameters->bInvertedAlpha = static_cast<uint32>(Metadata.bInvertedAlpha);
+			Parameters->RenderTargets[0] = Output.GetRenderTargetBinding();
+
+			TShaderMapRef<FCompositePassBlurShader> PixelShader(GlobalShaderMap, PermutationVector);
+			FPixelShaderUtils::AddFullscreenPass(
+				GraphBuilder,
+				GlobalShaderMap,
+				RDG_EVENT_NAME("Composite.Blur.Horizontal (%dx%d)", Output.ViewRect.Width(), Output.ViewRect.Height()),
+				PixelShader,
+				Parameters,
+				Output.ViewRect
+			);
+		}
+
+		return UE::CompositeCore::FPassTexture{ MoveTemp(Output), Metadata };
 	}
 }
 
@@ -134,7 +120,7 @@ bool UCompositePassBlur::GetIsActive() const
 
 FCompositeCorePassProxy* UCompositePassBlur::GetProxy(const UE::CompositeCore::FPassInputDecl& InputDecl, FCompositeTraversalContext& InContext, FSceneRenderingBulkObjectAllocator& InFrameAllocator) const
 {
-	using namespace UE::Composite;
+	using namespace UE::Composite::Private;
 
 	FCompositePassBlurProxy* Proxy = InFrameAllocator.Create<FCompositePassBlurProxy>(UE::CompositeCore::FPassInputDeclArray{ InputDecl });
 	Proxy->Radius = Radius;
diff --git a/Engine/Plugins/Compositing/Composite/Source/Composite/Private/Passes/CompositePassCenteredScale.cpp b/Engine/Plugins/Compositing/Composite/Source/Composite/Private/Passes/CompositePassCenteredScale.cpp
--- a/Engine/Plugins/Compositing/Composite/Source/Composite/Private/Passes/CompositePassCenteredScale.cpp
+++ b/Engine/Plugins/Compositing/Composite/Source/Composite/Private/Passes/CompositePassCenteredScale.cpp
@@ -32,63 +32,52 @@ class FCompositePassCenteredScaleShader : public FGlobalShader
 IMPLEMENT_GLOBAL_SHADER(FCompositePassCenteredScaleShader, "/Plugin/Composite/Private/CompositeCenteredScale.usf", "MainPS", SF_Pixel);
 
 
-namespace UE
+namespace UE::Composite::Private
 {
-	namespace CompositeCore
-	{
-		class FCenteredScalePassProxy : public FCompositeCorePassProxy
-		{
-		public:
-			IMPLEMENT_COMPOSITE_PASS(FCenteredScalePassProxy);
-
-			using FCompositeCorePassProxy::FCompositeCorePassProxy;
+	using namespace CompositeCore;
 
-			FPassTexture Add(FRDGBuilder& GraphBuilder, const FSceneView& InView, const FPassInputArray& Inputs, const FPassContext& PassContext) const override
-			{
-				RDG_EVENT_SCOPE_STAT(GraphBuilder, FCompositeCenteredScale, "Composite.CenteredScale");
-
-				check(ValidateInputs(Inputs));
+	FPassTexture FCenteredScalePassProxy::Add(FRDGBuilder& GraphBuilder, const FSceneView& InView, const FPassInputArray& Inputs, const FPassContext& PassContext) const
+	{
+		RDG_EVENT_SCOPE_STAT(GraphBuilder, FCompositeCenteredScale, "Composite.CenteredScale");
 
-				const UE::CompositeCore::FResourceMetadata& Metadata = Inputs[0].Metadata;
-				const FScreenPassTexture& Input = Inputs[0].Texture;
-				FScreenPassRenderTarget Output = Inputs.OverrideOutput;
+		check(ValidateInputs(Inputs));
 
-				// If the override output is provided, it means that this is the last pass in post processing.
-				if (!Output.IsValid())
-				{
-					Output = CreateOutputRenderTarget(GraphBuilder, InView, PassContext.OutputViewRect, Input.Texture->Desc, TEXT("CenteredScaleCompositePass"));
-				}
+		const UE::CompositeCore::FResourceMetadata& Metadata = Inputs[0].Metadata;
+		const FScreenPassTexture& Input = Inputs[0].Texture;
+		FScreenPassRenderTarget Output = Inputs.OverrideOutput;
 
-				FGlobalShaderMap* GlobalShaderMap = GetGlobalShaderMap(InView.GetFeatureLevel());
-				FCompositePassCenteredScaleShader::FParameters* Parameters = GraphBuilder.AllocParameters<FCompositePassCenteredScaleShader::FParameters>();
-				FRHISamplerState* InputSampler = Metadata.GetSamplerState();
-				Parameters->Input = GetScreenPassTextureInput(Input, InputSampler);
-				Parameters->SvPositionToOutputViewportUV = FScreenTransform::ChangeTextureBasisFromTo(
-					FScreenPassTextureViewport(Output),
-					FScreenTransform::ETextureBasis::TexelPosition,
-					FScreenTransform::ETextureBasis::ViewportUV);
-				Parameters->OutputViewportUVToInputTextureUV = FScreenTransform::ChangeTextureBasisFromTo(
-					FScreenPassTextureViewport(Input),
-					FScreenTransform::ETextureBasis::ViewportUV,
-					FScreenTransform::ETextureBasis::TextureUV);
-				Parameters->ScaleUV = ScaleUV;
-				Parameters->RenderTargets[0] = Output.GetRenderTargetBinding();
-
-				TShaderMapRef<FCompositePassCenteredScaleShader> PixelShader(GlobalShaderMap);
-				FPixelShaderUtils::AddFullscreenPass(
-					GraphBuilder,
-					GlobalShaderMap,
-					RDG_EVENT_NAME("Composite.CenteredScale (%dx%d)", Output.ViewRect.Width(), Output.ViewRect.Height()),
-					PixelShader,
-					Parameters,
-					Output.ViewRect
-				);
-
-				return UE::CompositeCore::FPassTexture{ MoveTemp(Output), Metadata };
-			}
+		// If the override output is provided, it means that this is the last pass in post processing.
+		if (!Output.IsValid())
+		{
+			Output = CreateOutputRenderTarget(GraphBuilder, InView, PassContext.OutputViewRect, Input.Texture->Desc, TEXT("CenteredScaleCompositePass"));
+		}
 
-			FVector2f ScaleUV = FVector2f::One();
-		};
+		FGlobalShaderMap* GlobalShaderMap = GetGlobalShaderMap(InView.GetFeatureLevel());
+		FCompositePassCenteredScaleShader::FParameters* Parameters = GraphBuilder.AllocParameters<FCompositePassCenteredScaleShader::FParameters>();
+		FRHISamplerState* InputSampler = Metadata.GetSamplerState();
+		Parameters->Input = GetScreenPassTextureInput(Input, InputSampler);
+		Parameters->SvPositionToOutputViewportUV = FScreenTransform::ChangeTextureBasisFromTo(
+			FScreenPassTextureViewport(Output),
+			FScreenTransform::ETextureBasis::TexelPosition,
+			FScreenTransform::ETextureBasis::ViewportUV);
+		Parameters->OutputViewportUVToInputTextureUV = FScreenTransform::ChangeTextureBasisFromTo(
+			FScreenPassTextureViewport(Input),
+			FScreenTransform::ETextureBasis::ViewportUV,
+			FScreenTransform::ETextureBasis::TextureUV);
+		Parameters->ScaleUV = ScaleUV;
+		Parameters->RenderTargets[0] = Output.GetRenderTargetBinding();
+
+		TShaderMapRef<FCompositePassCenteredScaleShader> PixelShader(GlobalShaderMap);
+		FPixelShaderUtils::AddFullscreenPass(
+			GraphBuilder,
+			GlobalShaderMap,
+			RDG_EVENT_NAME("Composite.CenteredScale (%dx%d)", Output.ViewRect.Width(), Output.ViewRect.Height()),
+			PixelShader,
+			Parameters,
+			Output.ViewRect
+		);
+
+		return UE::CompositeCore::FPassTexture{ MoveTemp(Output), Metadata };
 	}
 }
 
@@ -108,7 +97,7 @@ UCompositePassCenteredScale::~UCompositePassCenteredScale() = default;
 
 FCompositeCorePassProxy* UCompositePassCenteredScale::GetProxy(const UE::CompositeCore::FPassInputDecl& InputDecl, FCompositeTraversalContext& InContext, FSceneRenderingBulkObjectAllocator& InFrameAllocator) const
 {
-	using namespace UE::CompositeCore;
+	using namespace UE::Composite::Private;
 
 	// Use cached value from GetIsActive() if available, otherwise compute fresh.
 	const FVector2f ScaleUV = CachedScaleUV.IsSet() ? CachedScaleUV.GetValue() : CalculateScale();
diff --git a/Engine/Plugins/Compositing/Composite/Source/Composite/Private/Passes/CompositePassColorKeyer.cpp b/Engine/Plugins/Compositing/Composite/Source/Composite/Private/Passes/CompositePassColorKeyer.cpp
--- a/Engine/Plugins/Compositing/Composite/Source/Composite/Private/Passes/CompositePassColorKeyer.cpp
+++ b/Engine/Plugins/Compositing/Composite/Source/Composite/Private/Passes/CompositePassColorKeyer.cpp
@@ -75,143 +75,123 @@ class FCompositePassColorKeyerShader : public FGlobalShader
 IMPLEMENT_GLOBAL_SHADER(FCompositePassDenoiseShader, "/Plugin/Composite/Private/CompositeDenoise.usf", "MainCS", SF_Compute);
 IMPLEMENT_GLOBAL_SHADER(FCompositePassColorKeyerShader, "/Plugin/Composite/Private/CompositeColorKeyer.usf", "MainCS", SF_Compute);
 
-namespace UE
+static FScreenPassTexture AddColorKeyerDenoisePass(FRDGBuilder& GraphBuilder, const FScreenPassTexture& Input, ERHIFeatureLevel::Type FeatureLevel)
 {
-	namespace CompositeCore
-	{
-		class FCompositePassColorKeyerProxy : public FCompositeCorePassProxy
-		{
-		public:
-			IMPLEMENT_COMPOSITE_PASS(FCompositePassColorKeyerProxy);
-
-			using FCompositeCorePassProxy::FCompositeCorePassProxy;
+	RDG_EVENT_SCOPE_STAT(GraphBuilder, FCompositeDenoise, "Composite.Denoise");
+
+	FGlobalShaderMap* GlobalShaderMap = GetGlobalShaderMap(FeatureLevel);
+	const FIntPoint TextureSize = Input.ViewRect.Size();
+
+	const FRDGTextureDesc OutputDesc = FRDGTextureDesc::Create2D(TextureSize, PF_FloatRGBA, FClearValueBinding::Black, TexCreate_UAV | TexCreate_ShaderResource | TexCreate_RenderTargetable);
+	FScreenPassRenderTarget Output = FScreenPassRenderTarget(GraphBuilder.CreateTexture(OutputDesc, TEXT("CompositeTextureDenoised")), FIntRect(FIntPoint::ZeroValue, TextureSize), ERenderTargetLoadAction::ENoAction);
+
+	FCompositePassDenoiseShader::FParameters* PassParameters = GraphBuilder.AllocParameters<FCompositePassDenoiseShader::FParameters>();
+	// Always use point sampling: denoising needs exact texel values for its spatial filter.
+	PassParameters->Input = GetScreenPassTextureInput(Input, TStaticSamplerState<SF_Point>::GetRHI());
+	PassParameters->RWOutputTexture = GraphBuilder.CreateUAV(Output.Texture);
+
+	TShaderMapRef<FCompositePassDenoiseShader> ComputeShader(GlobalShaderMap);
+	FComputeShaderUtils::AddPass(
+		GraphBuilder,
+		RDG_EVENT_NAME("Composite.Denoise (%dx%d)", TextureSize.X, TextureSize.Y),
+		ERDGPassFlags::Compute,
+		ComputeShader,
+		PassParameters,
+		FComputeShaderUtils::GetGroupCount(TextureSize, FCompositePassDenoiseShader::ThreadGroupSize)
+	);
+
+	return Output;
+}
 
-			FScreenPassTexture AddDenoisePass(FRDGBuilder& GraphBuilder, const FScreenPassTexture& Input, ERHIFeatureLevel::Type FeatureLevel) const
-			{
-				RDG_EVENT_SCOPE_STAT(GraphBuilder, FCompositeDenoise, "Composite.Denoise");
-
-				FGlobalShaderMap* GlobalShaderMap = GetGlobalShaderMap(FeatureLevel);
-				const FIntPoint TextureSize = Input.ViewRect.Size();
-
-				const FRDGTextureDesc OutputDesc = FRDGTextureDesc::Create2D(TextureSize, PF_FloatRGBA, FClearValueBinding::Black, TexCreate_UAV | TexCreate_ShaderResource | TexCreate_RenderTargetable);
-				FScreenPassRenderTarget Output = FScreenPassRenderTarget(GraphBuilder.CreateTexture(OutputDesc, TEXT("CompositeTextureDenoised")), FIntRect(FIntPoint::ZeroValue, TextureSize), ERenderTargetLoadAction::ENoAction);
-
-				FCompositePassDenoiseShader::FParameters* PassParameters = GraphBuilder.AllocParameters<FCompositePassDenoiseShader::FParameters>();
-				// Always use point sampling: denoising needs exact texel values for its spatial filter.
-				PassParameters->Input = GetScreenPassTextureInput(Input, TStaticSamplerState<SF_Point>::GetRHI());
-				PassParameters->RWOutputTexture = GraphBuilder.CreateUAV(Output.Texture);
-
-				TShaderMapRef<FCompositePassDenoiseShader> ComputeShader(GlobalShaderMap);
-				FComputeShaderUtils::AddPass(
-					GraphBuilder,
-					RDG_EVENT_NAME("Composite.Denoise (%dx%d)", TextureSize.X, TextureSize.Y),
-					ERDGPassFlags::Compute,
-					ComputeShader,
-					PassParameters,
-					FComputeShaderUtils::GetGroupCount(TextureSize, FCompositePassDenoiseShader::ThreadGroupSize)
-				);
-
-				return Output;
-			}
+namespace UE::Composite::Private
+{
+	using namespace CompositeCore;
 
-			FPassTexture Add(FRDGBuilder& GraphBuilder, const FSceneView& InView, const FPassInputArray& Inputs, const FPassContext& PassContext) const override
-			{
-				RDG_EVENT_SCOPE_STAT(GraphBuilder, FCompositeColorKeyer, "Composite.ColorKeyer");
+	FPassTexture FCompositePassColorKeyerProxy::Add(FRDGBuilder& GraphBuilder, const FSceneView& InView, const FPassInputArray& Inputs, const FPassContext& PassContext) const
+	{
+		RDG_EVENT_SCOPE_STAT(GraphBuilder, FCompositeColorKeyer, "Composite.ColorKeyer");
 
-				check(ValidateInputs(Inputs));
+		check(ValidateInputs(Inputs));
 
-				const UE::CompositeCore::FResourceMetadata& Metadata = Inputs[0].Metadata;
-				FScreenPassTexture Input = Inputs[0].Texture;
-				FScreenPassRenderTarget Output = Inputs.OverrideOutput;
+		const UE::CompositeCore::FResourceMetadata& Metadata = Inputs[0].Metadata;
+		FScreenPassTexture Input = Inputs[0].Texture;
+		FScreenPassRenderTarget Output = Inputs.OverrideOutput;
 
-				if (DenoiseMethod != ECompositeDenoiseMethod::None)
-				{
-					Input = AddDenoisePass(GraphBuilder, Input, InView.GetFeatureLevel());
-				}
+		if (DenoiseMethod != ECompositeDenoiseMethod::None)
+		{
+			Input = AddColorKeyerDenoisePass(GraphBuilder, Input, InView.GetFeatureLevel());
+		}
 
-				// If the override output is provided, it means that this is the last pass in post processing.
-				if (!Output.IsValid())
-				{
-					Output = CreateOutputRenderTarget(GraphBuilder, InView, PassContext.OutputViewRect, Input.Texture->Desc, TEXT("ColorKeyerCompositePass"));
-				}
+		// If the override output is provided, it means that this is the last pass in post processing.
+		if (!Output.IsValid())
+		{
+			Output = CreateOutputRenderTarget(GraphBuilder, InView, PassContext.OutputViewRect, Input.Texture->Desc, TEXT("ColorKeyerCompositePass"));
+		}
 
-				FGlobalShaderMap* GlobalShaderMap = GetGlobalShaderMap(InView.GetFeatureLevel());
-
-				const FScreenPassTextureViewportParameters OutputParameters = GetScreenPassTextureViewportParameters(FScreenPassTextureViewport(Output));
-
-				FRHISamplerState* InputSampler = Metadata.GetSamplerState();
-
-				FCompositePassColorKeyerShader::FParameters* PassParameters = GraphBuilder.AllocParameters<FCompositePassColorKeyerShader::FParameters>();
-				PassParameters->Input = GetScreenPassTextureInput(Input, InputSampler);
-				PassParameters->RWOutputTexture = GraphBuilder.CreateUAV(Output.Texture);
-				PassParameters->Output = OutputParameters;
-				PassParameters->SvPositionToInputTextureUV = (
-					FScreenTransform::ChangeTextureBasisFromTo(FScreenPassTextureViewport(Output), FScreenTransform::ETextureBasis::TexelPosition, FScreenTransform::ETextureBasis::ViewportUV) *
-					FScreenTransform::ChangeTextureBasisFromTo(FScreenPassTextureViewport(Input), FScreenTransform::ETextureBasis::ViewportUV, FScreenTransform::ETextureBasis::TextureUV));
-				PassParameters->KeyColor = FVector3f(KeyColor);
-				PassParameters->Params0 = Params0;
-				PassParameters->Params1 = Params1;
-				PassParameters->ScreenType = static_cast<uint32>(ScreenType);
-				PassParameters->Visualization = static_cast<uint32>(Visualization);
-				PassParameters->bPreserveVignetteAfterKey = static_cast<uint32>(bPreserveVignetteAfterKey);
-				PassParameters->bInvertAlpha = static_cast<uint32>(bInvertAlpha);
-	
-				FRDGTextureRef CleanPlateResource = nullptr;
-				PassParameters->CleanPlate.Texture = GSystemTextures.GetBlackDummy(GraphBuilder);
-				// Clean plate is a separate reference texture not driven by the input metadata; bilinear is always appropriate.
-				PassParameters->CleanPlate.Sampler = TStaticSamplerState<SF_Bilinear>::GetRHI();
-
-				TStrongObjectPtr<UTexture> CleanPlateTexturePtr = CleanPlateWeakPtr.Pin();
-				if (CleanPlateTexturePtr.IsValid())
+		FGlobalShaderMap* GlobalShaderMap = GetGlobalShaderMap(InView.GetFeatureLevel());
+
+		const FScreenPassTextureViewportParameters OutputParameters = GetScreenPassTextureViewportParameters(FScreenPassTextureViewport(Output));
+
+		FRHISamplerState* InputSampler = Metadata.GetSamplerState();
+
+		FCompositePassColorKeyerShader::FParameters* PassParameters = GraphBuilder.AllocParameters<FCompositePassColorKeyerShader::FParameters>();
+		PassParameters->Input = GetScreenPassTextureInput(Input, InputSampler);
+		PassParameters->RWOutputTexture = GraphBuilder.CreateUAV(Output.Texture);
+		PassParameters->Output = OutputParameters;
+		PassParameters->SvPositionToInputTextureUV = (
+			FScreenTransform::ChangeTextureBasisFromTo(FScreenPassTextureViewport(Output), FScreenTransform::ETextureBasis::TexelPosition, FScreenTransform::ETextureBasis::ViewportUV) *
+			FScreenTransform::ChangeTextureBasisFromTo(FScreenPassTextureViewport(Input), FScreenTransform::ETextureBasis::ViewportUV, FScreenTransform::ETextureBasis::TextureUV));
+		PassParameters->KeyColor = FVector3f(KeyColor);
+		PassParameters->Params0 = Params0;
+		PassParameters->Params1 = Params1;
+		PassParameters->ScreenType = static_cast<uint32>(ScreenType);
+		PassParameters->Visualization = static_cast<uint32>(Visualization);
+		PassParameters->bPreserveVignetteAfterKey = static_cast<uint32>(bPreserveVignetteAfterKey);
+		PassParameters->bInvertAlpha = static_cast<uint32>(bInvertAlpha);
+
+		FRDGTextureRef CleanPlateResource = nullptr;
+		PassParameters->CleanPlate.Texture = GSystemTextures.GetBlackDummy(GraphBuilder);
+		// Clean plate is a separate reference texture not driven by the input metadata; bilinear is always appropriate.
+		PassParameters->CleanPlate.Sampler = TStaticSamplerState<SF_Bilinear>::GetRHI();
+
+		TStrongObjectPtr<UTexture> CleanPlateTexturePtr = CleanPlateWeakPtr.Pin();
+		if (CleanPlateTexturePtr.IsValid())
+		{
+			const FTextureResource* TextureResource = CleanPlateTexturePtr->GetResource();
+			if (TextureResource != nullptr)
+			{
+				FRHITexture* TextureRHI = TextureResource->GetTexture2DRHI();
+				if (TextureRHI != nullptr)
 				{
-					const FTextureResource* TextureResource = CleanPlateTexturePtr->GetResource();
-					if (TextureResource != nullptr)
-					{
-						FRHITexture* TextureRHI = TextureResource->GetTexture2DRHI();
-						if (TextureRHI != nullptr)
-						{
-							CleanPlateResource = GraphBuilder.RegisterExternalTexture(CreateRenderTarget(TextureRHI, TEXT("CompositeCleanPlateTexture")));
-							GraphBuilder.UseInternalAccessMode(CleanPlateResource);
-
-							PassParameters->CleanPlate.Texture = CleanPlateResource;
-						}
-					}
-				}
+					CleanPlateResource = GraphBuilder.RegisterExternalTexture(CreateRenderTarget(TextureRHI, TEXT("CompositeCleanPlateTexture")));
+					GraphBuilder.UseInternalAccessMode(CleanPlateResource);
 
-				FCompositePassColorKeyerShader::FPermutationDomain PermutationVector;
-				PermutationVector.Set<FCompositePassColorKeyerShader::FUseCleanPlate>(CleanPlateResource != nullptr);
-				PermutationVector.Set<FCompositePassColorKeyerShader::FDespillOnly>(bDespillOnly);
-
-				const FIntPoint OutputSize = Output.ViewRect.Size();
-				TShaderMapRef<FCompositePassColorKeyerShader> ComputeShader(GlobalShaderMap, PermutationVector);
-				FComputeShaderUtils::AddPass(
-					GraphBuilder,
-					RDG_EVENT_NAME("Composite.ColorKeyer (%dx%d)", OutputSize.X, OutputSize.Y),
-					ERDGPassFlags::Compute,
-					ComputeShader,
-					PassParameters,
-					FComputeShaderUtils::GetGroupCount(OutputSize, FCompositePassColorKeyerShader::ThreadGroupSize)
-				);
-
-				if (CleanPlateResource != nullptr)
-				{
-					GraphBuilder.UseExternalAccessMode(CleanPlateResource, ERHIAccess::SRVMask);
+					PassParameters->CleanPlate.Texture = CleanPlateResource;
 				}
-
-				return UE::CompositeCore::FPassTexture{ MoveTemp(Output), Metadata };
 			}
+		}
+
+		FCompositePassColorKeyerShader::FPermutationDomain PermutationVector;
+		PermutationVector.Set<FCompositePassColorKeyerShader::FUseCleanPlate>(CleanPlateResource != nullptr);
+		PermutationVector.Set<FCompositePassColorKeyerShader::FDespillOnly>(bDespillOnly);
+
+		const FIntPoint OutputSize = Output.ViewRect.Size();
+		TShaderMapRef<FCompositePassColorKeyerShader> ComputeShader(GlobalShaderMap, PermutationVector);
+		FComputeShaderUtils::AddPass(
+			GraphBuilder,
+			RDG_EVENT_NAME("Composite.ColorKeyer (%dx%d)", OutputSize.X, OutputSize.Y),
+			ERDGPassFlags::Compute,
+			ComputeShader,
+			PassParameters,
+			FComputeShaderUtils::GetGroupCount(OutputSize, FCompositePassColorKeyerShader::ThreadGroupSize)
+		);
+
+		if (CleanPlateResource != nullptr)
+		{
+			GraphBuilder.UseExternalAccessMode(CleanPlateResource, ERHIAccess::SRVMask);
+		}
 
-			ECompositeColorKeyerScreenType ScreenType;
-			FLinearColor KeyColor;
-			FVector4f Params0;
-			FVector4f Params1;
-			ECompositeColorKeyerVisualization Visualization;
-			bool bPreserveVignetteAfterKey;
-			ECompositeDenoiseMethod DenoiseMethod;
-			bool bDespillOnly;
-			bool bInvertAlpha;
-			TWeakObjectPtr<UTexture> CleanPlateWeakPtr;
-		};
+		return UE::CompositeCore::FPassTexture{ MoveTemp(Output), Metadata };
 	}
 }
 
@@ -238,7 +218,7 @@ UCompositePassColorKeyer::~UCompositePassColorKeyer() = default;
 
 FCompositeCorePassProxy* UCompositePassColorKeyer::GetProxy(const UE::CompositeCore::FPassInputDecl& InputDecl, FCompositeTraversalContext& InContext, FSceneRenderingBulkObjectAllocator& InFrameAllocator) const
 {
-	using namespace UE::CompositeCore;
+	using namespace UE::Composite::Private;
 
 	FCompositePassColorKeyerProxy* Proxy = InFrameAllocator.Create<FCompositePassColorKeyerProxy>(UE::CompositeCore::FPassInputDeclArray{ InputDecl });
 	Proxy->KeyColor = KeyColor;
diff --git a/Engine/Plugins/Compositing/Composite/Source/Composite/Private/Passes/CompositePassMasking.cpp b/Engine/Plugins/Compositing/Composite/Source/Composite/Private/Passes/CompositePassMasking.cpp
--- a/Engine/Plugins/Compositing/Composite/Source/Composite/Private/Passes/CompositePassMasking.cpp
+++ b/Engine/Plugins/Compositing/Composite/Source/Composite/Private/Passes/CompositePassMasking.cpp
@@ -38,78 +38,62 @@ class FCompositePassMaskingShader : public FGlobalShader
 IMPLEMENT_GLOBAL_SHADER(FCompositePassMaskingShader, "/Plugin/Composite/Private/CompositeMasking.usf", "MainPS", SF_Pixel);
 
 
-namespace UE
+namespace UE::Composite::Private
 {
-	namespace Composite
+	using namespace CompositeCore;
+
+	static FCompositePassMaskingShader::FInputPassTextureParameters GetPassTextureParameters(const UE::CompositeCore::FPassInput& Input)
+	{
+		FCompositePassMaskingShader::FInputPassTextureParameters Parameters;
+		Parameters.Texture = Input.Texture.Texture;
+		Parameters.Sampler = Input.Metadata.GetSamplerState();
+		Parameters.bInvertedAlpha = static_cast<uint32>(Input.Metadata.bInvertedAlpha);
+		Parameters.SourceEncoding = static_cast<uint32>(Input.Metadata.Encoding);
+		return Parameters;
+	}
+
+	FPassTexture FMaskingPassProxy::Add(FRDGBuilder& GraphBuilder, const FSceneView& InView, const FPassInputArray& Inputs, const FPassContext& PassContext) const
 	{
-		using namespace CompositeCore;
+		RDG_EVENT_SCOPE_STAT(GraphBuilder, FCompositeMasking, "Composite.Masking");
+
+		check(ValidateInputs(Inputs));
+
+		const UE::CompositeCore::FResourceMetadata& Metadata = Inputs[0].Metadata;
+		const FScreenPassTexture& InputTexture = Inputs[0].Texture;
+		FScreenPassRenderTarget Output = Inputs.OverrideOutput;
 
-		class FMaskingPassProxy : public FCompositeCorePassProxy
+		// If the override output is provided, it means that this is the last pass in post processing.
+		if (!Output.IsValid())
 		{
-		public:
-			IMPLEMENT_COMPOSITE_PASS(FMaskingPassProxy);
-
-			using FCompositeCorePassProxy::FCompositeCorePassProxy;
-
-			FCompositePassMaskingShader::FInputPassTextureParameters GetPassTextureParameters(const UE::CompositeCore::FPassInput& Input) const
-			{
-				FCompositePassMaskingShader::FInputPassTextureParameters Parameters;
-				Parameters.Texture = Input.Texture.Texture;
-				Parameters.Sampler = Input.Metadata.GetSamplerState();
-				Parameters.bInvertedAlpha = static_cast<uint32>(Input.Metadata.bInvertedAlpha);
-				Parameters.SourceEncoding = static_cast<uint32>(Input.Metadata.Encoding);
-
-				return Parameters;
-			}
-
-			FPassTexture Add(FRDGBuilder& GraphBuilder, const FSceneView& InView, const FPassInputArray& Inputs, const FPassContext& PassContext) const override
-			{
-				RDG_EVENT_SCOPE_STAT(GraphBuilder, FCompositeMasking, "Composite.Masking");
-
-				check(ValidateInputs(Inputs));
-
-				const UE::CompositeCore::FResourceMetadata& Metadata = Inputs[0].Metadata;
-				const FScreenPassTexture& InputTexture = Inputs[0].Texture;
-				FScreenPassRenderTarget Output = Inputs.OverrideOutput;
-
-				// If the override output is provided, it means that this is the last pass in post processing.
-				if (!Output.IsValid())
-				{
-					Output = CreateOutputRenderTarget(GraphBuilder, InView, PassContext.OutputViewRect, InputTexture.Texture->Desc, TEXT("MaskingCompositePass"));
-				}
-
-				FGlobalShaderMap* GlobalShaderMap = GetGlobalShaderMap(InView.GetFeatureLevel());
-				FCompositePassMaskingShader::FParameters* Parameters = GraphBuilder.AllocParameters<FCompositePassMaskingShader::FParameters>();
-				Parameters->Input = GetPassTextureParameters(Inputs[0]);
-				Parameters->Mask = GetPassTextureParameters(Inputs[1]);
-				Parameters->MaskSourceChannel = MaskSourceChannel;
-				Parameters->bUnpremult = static_cast<uint32>(bUnpremult);
-				Parameters->bInvert = static_cast<uint32>(bInvert);
-				Parameters->SvPositionToInputTextureUV = (
-					FScreenTransform::ChangeTextureBasisFromTo(FScreenPassTextureViewport(Output), FScreenTransform::ETextureBasis::TexelPosition, FScreenTransform::ETextureBasis::ViewportUV) *
-					FScreenTransform::ChangeTextureBasisFromTo(FScreenPassTextureViewport(Inputs[0].Texture), FScreenTransform::ETextureBasis::ViewportUV, FScreenTransform::ETextureBasis::TextureUV));
-				Parameters->SvPositionToMaskTextureUV = (
-					FScreenTransform::ChangeTextureBasisFromTo(FScreenPassTextureViewport(Output), FScreenTransform::ETextureBasis::TexelPosition, FScreenTransform::ETextureBasis::ViewportUV) *
-					FScreenTransform::ChangeTextureBasisFromTo(FScreenPassTextureViewport(Inputs[1].Texture), FScreenTransform::ETextureBasis::ViewportUV, FScreenTransform::ETextureBasis::TextureUV));
-				Parameters->RenderTargets[0] = Output.GetRenderTargetBinding();
-
-				TShaderMapRef<FCompositePassMaskingShader> PixelShader(GlobalShaderMap);
-				FPixelShaderUtils::AddFullscreenPass(
-					GraphBuilder,
-					GlobalShaderMap,
-					RDG_EVENT_NAME("Composite.Masking (%dx%d)", Output.ViewRect.Width(), Output.ViewRect.Height()),
-					PixelShader,
-					Parameters,
-					Output.ViewRect
-				);
-
-				return UE::CompositeCore::FPassTexture{ MoveTemp(Output), Metadata };
-			}
-
-			uint32 MaskSourceChannel = 0;
-			bool bUnpremult = true;
-			bool bInvert = false;
-		};
+			Output = CreateOutputRenderTarget(GraphBuilder, InView, PassContext.OutputViewRect, InputTexture.Texture->Desc, TEXT("MaskingCompositePass"));
+		}
+
+		FGlobalShaderMap* GlobalShaderMap = GetGlobalShaderMap(InView.GetFeatureLevel());
+		FCompositePassMaskingShader::FParameters* Parameters = GraphBuilder.AllocParameters<FCompositePassMaskingShader::FParameters>();
+		Parameters->Input = GetPassTextureParameters(Inputs[0]);
+		Parameters->Mask = GetPassTextureParameters(Inputs[1]);
+		Parameters->MaskSourceChannel = MaskSourceChannel;
+		Parameters->bUnpremult = static_cast<uint32>(bUnpremult);
+		Parameters->bInvert = static_cast<uint32>(bInvert);
+		Parameters->SvPositionToInputTextureUV = (
+			FScreenTransform::ChangeTextureBasisFromTo(FScreenPassTextureViewport(Output), FScreenTransform::ETextureBasis::TexelPosition, FScreenTransform::ETextureBasis::ViewportUV) *
+			FScreenTransform::ChangeTextureBasisFromTo(FScreenPassTextureViewport(Inputs[0].Texture), FScreenTransform::ETextureBasis::ViewportUV, FScreenTransform::ETextureBasis::TextureUV));
+		Parameters->SvPositionToMaskTextureUV = (
+			FScreenTransform::ChangeTextureBasisFromTo(FScreenPassTextureViewport(Output), FScreenTransform::ETextureBasis::TexelPosition, FScreenTransform::ETextureBasis::ViewportUV) *
+			FScreenTransform::ChangeTextureBasisFromTo(FScreenPassTextureViewport(Inputs[1].Texture), FScreenTransform::ETextureBasis::ViewportUV, FScreenTransform::ETextureBasis::TextureUV));
+		Parameters->RenderTargets[0] = Output.GetRenderTargetBinding();
+
+		TShaderMapRef<FCompositePassMaskingShader> PixelShader(GlobalShaderMap);
+		FPixelShaderUtils::AddFullscreenPass(
+			GraphBuilder,
+			GlobalShaderMap,
+			RDG_EVENT_NAME("Composite.Masking (%dx%d)", Output.ViewRect.Width(), Output.ViewRect.Height()),
+			PixelShader,
+			Parameters,
+			Output.ViewRect
+		);
+
+		return UE::CompositeCore::FPassTexture{ MoveTemp(Output), Metadata };
 	}
 }
 
@@ -131,7 +115,7 @@ bool UCompositePassMasking::GetIsActive() const
 
 FCompositeCorePassProxy* UCompositePassMasking::GetProxy(const UE::CompositeCore::FPassInputDecl& InputDecl, FCompositeTraversalContext& InContext, FSceneRenderingBulkObjectAllocator& InFrameAllocator) const
 {
-	using namespace UE::Composite;
+	using namespace UE::Composite::Private;
 
 	if (!ensureMsgf(IsValid(MaskTexture), TEXT("Texture is expected to be valid for active pass.")))
 	{
diff --git a/Engine/Plugins/Compositing/Composite/Source/Composite/Private/Tests/CompositePassBlurTest.cpp b/Engine/Plugins/Compositing/Composite/Source/Composite/Private/Tests/CompositePassBlurTest.cpp
new file mode 100644
--- /dev/null
+++ b/Engine/Plugins/Compositing/Composite/Source/Composite/Private/Tests/CompositePassBlurTest.cpp
@@ -0,0 +1,455 @@
+// Copyright Epic Games, Inc. All Rights Reserved.
+
+#include "CoreTypes.h"
+#include "RenderGraphBuilder.h"
+#include "RHI.h"
+#include "ScreenPass.h"
+
+#include "Passes/CompositePassBlur.h"
+#include "Tests/CompositeTestHelpers.h"
+
+#if WITH_DEV_AUTOMATION_TESTS
+#include "Misc/AutomationTest.h"
+#endif
+
+namespace UE::Composite::Tests
+{
+
+#if WITH_DEV_AUTOMATION_TESTS
+
+// -- Upload/download parameter structs (shared by both tests) --
+
+BEGIN_SHADER_PARAMETER_STRUCT(FBlurTestUploadParameters, )
+	RDG_TEXTURE_ACCESS(InputTexture, ERHIAccess::CopyDest)
+END_SHADER_PARAMETER_STRUCT()
+
+BEGIN_SHADER_PARAMETER_STRUCT(FBlurTestDownloadParameters, )
+	RDG_TEXTURE_ACCESS(OutputTexture, ERHIAccess::CopySrc)
+END_SHADER_PARAMETER_STRUCT()
+
+// ============================================================================
+// Test 1: uniform-field
+// ============================================================================
+
+namespace Blur  // unique name prevents symbol collision in unity builds
+{
+
+// A uniform-color texture yields output == input.
+// This holds regardless of kernel radius or separable-pass order.
+static FLinearColor BlurUniformCPU(FLinearColor In)
+{
+	return In;
+}
+
+struct FBlurScenario
+{
+	const TCHAR* Label;
+	FLinearColor  Input;
+	int32         Radius;
+	bool          bAlphaOnly;
+};
+
+static const FBlurScenario GScenarios[] =
+{
+	// Uniform inputs: blurring any constant field returns the same constant.
+	{ TEXT("UniformSmallRadius"),  { 0.5f, 0.3f, 0.7f, 0.8f }, 4,  false },
+	{ TEXT("UniformLargeRadius"),  { 0.2f, 0.6f, 0.4f, 1.0f }, 16, false },
+	{ TEXT("AlphaOnlyUniform"),    { 0.8f, 0.4f, 0.2f, 0.6f }, 8,  true  },
+};
+
+}  // namespace Blur
+
+IMPLEMENT_SIMPLE_AUTOMATION_TEST(FCompositeBlurPassTest, "Composite.UnitTests.BlurPass",
+	EAutomationTestFlags_ApplicationContextMask | EAutomationTestFlags::EngineFilter | EAutomationTestFlags::NonNullRHI)
+
+bool FCompositeBlurPassTest::RunTest(const FString& Parameter)
+{
+	using namespace UE::Composite::Private;
+	using namespace UE::CompositeCore;
+	using namespace Blur;
+
+	const int32 NumScenarios = UE_ARRAY_COUNT(GScenarios);
+
+	TArray<FLinearColor> ResultRDG;
+	ResultRDG.SetNumZeroed(NumScenarios);
+
+	TArray<FCompositePassBlurProxy> Proxies;
+	Proxies.Reserve(NumScenarios);
+	for (int32 i = 0; i < NumScenarios; ++i)
+	{
+		const FBlurScenario& S = GScenarios[i];
+		FCompositePassBlurProxy Proxy(GetDefaultInputDeclArray());
+		Proxy.Radius    = S.Radius;
+		Proxy.bAlphaOnly = S.bAlphaOnly;
+		Proxies.Emplace(MoveTemp(Proxy));
+	}
+
+	FCompositeTestView TestView;
+	DispatchAndWait(TestView.Get(),
+		[&ResultRDG, NumScenarios, Proxies = MoveTemp(Proxies)](FRDGBuilder& GraphBuilder, const FSceneView& View) mutable
+		{
+			const FIntRect PixelRect(0, 0, GTestViewSize.X, GTestViewSize.Y);
+
+			for (int32 i = 0; i < NumScenarios; ++i)
+			{
+				const FBlurScenario& S = GScenarios[i];
+
+				const FRDGTextureDesc InputDesc = FRDGTextureDesc::Create2D(
+					GTestViewSize, PF_A32B32G32R32F,
+					FClearValueBinding::Black, ETextureCreateFlags::None);
+				FRDGTextureRef InputTexture = GraphBuilder.CreateTexture(InputDesc, TEXT("BlurTest.Input"));
+
+				const FRDGTextureDesc OutputDesc = FRDGTextureDesc::Create2D(
+					GTestViewSize, PF_A32B32G32R32F,
+					FClearValueBinding::Black, ETextureCreateFlags::RenderTargetable);
+				FRDGTextureRef OutputTexture = GraphBuilder.CreateTexture(OutputDesc, TEXT("BlurTest.Output"));
+
+				{
+					FBlurTestUploadParameters* Params = GraphBuilder.AllocParameters<FBlurTestUploadParameters>();
+					Params->InputTexture = InputTexture;
+
+					const FLinearColor PixelValue = S.Input;
+					GraphBuilder.AddPass(
+						RDG_EVENT_NAME("BlurTest.Upload[%d]", i), Params,
+						ERDGPassFlags::Readback,
+						[PixelValue, Params](FRHICommandListImmediate& RHICmdList)
+						{
+							uint32 Stride = 0;
+							void* Data = RHICmdList.LockTexture2D(Params->InputTexture->GetRHI(), 0, RLM_WriteOnly, Stride, false);
+							for (int32 Row = 0; Row < GTestViewSize.Y; ++Row)
+							{
+								FLinearColor* RowPtr = reinterpret_cast<FLinearColor*>(static_cast<uint8*>(Data) + Row * Stride);
+								for (int32 Col = 0; Col < GTestViewSize.X; ++Col)
+								{
+									RowPtr[Col] = PixelValue;
+								}
+							}
+							RHICmdList.UnlockTexture2D(Params->InputTexture->GetRHI(), 0, false);
+						});
+				}
+
+				FPassInputArray Inputs;
+				{
+					FPassInput& In = Inputs.GetArray().AddDefaulted_GetRef();
+					In.Texture         = FScreenPassTexture{ InputTexture, PixelRect };
+					In.Metadata.Filter = SF_Point;
+				}
+				Inputs.OverrideOutput = FScreenPassRenderTarget{
+					FScreenPassTexture{ OutputTexture, PixelRect }, ERenderTargetLoadAction::ENoAction };
+
+				FPassContext PassContext;
+				PassContext.OutputViewRect = PixelRect;
+
+				Proxies[i].Add(GraphBuilder, View, Inputs, PassContext);
+
+				{
+					FBlurTestDownloadParameters* Params = GraphBuilder.AllocParameters<FBlurTestDownloadParameters>();
+					Params->OutputTexture = OutputTexture;
+
+					FLinearColor* ResultPtr = &ResultRDG[i];
+					GraphBuilder.AddPass(
+						RDG_EVENT_NAME("BlurTest.Download[%d]", i), Params,
+						ERDGPassFlags::Readback,
+						[ResultPtr, Params](FRHICommandListImmediate& RHICmdList)
+						{
+							uint32 Stride = 0;
+							const void* Data = RHICmdList.LockTexture2D(Params->OutputTexture->GetRHI(), 0, RLM_ReadOnly, Stride, false);
+							const FLinearColor* RowPtr = reinterpret_cast<const FLinearColor*>(
+								static_cast<const uint8*>(Data) + (GTestViewSize.Y / 2) * Stride);
+							FMemory::Memcpy(ResultPtr, &RowPtr[GTestViewSize.X / 2], sizeof(FLinearColor));
+							RHICmdList.UnlockTexture2D(Params->OutputTexture->GetRHI(), 0, false);
+						});
+				}
+			}
+		});
+
+	// Blur uses bilinear multi-tap sampling across two separable passes, accumulating
+	// more floating-point rounding than single-operation passes. 5e-4 is still tight
+	// enough to catch shader errors while tolerating hardware summation noise.
+	constexpr float Tolerance = 5e-4f;
+	for (int32 i = 0; i < NumScenarios; ++i)
+	{
+		const FBlurScenario& S = GScenarios[i];
+		const FLinearColor Expected = BlurUniformCPU(S.Input);
+		const FLinearColor& Actual  = ResultRDG[i];
+
+		UTEST_EQUAL_TOLERANCE(S.Label, Actual.R, Expected.R, Tolerance);
+		UTEST_EQUAL_TOLERANCE(S.Label, Actual.G, Expected.G, Tolerance);
+		UTEST_EQUAL_TOLERANCE(S.Label, Actual.B, Expected.B, Tolerance);
+		UTEST_EQUAL_TOLERANCE(S.Label, Actual.A, Expected.A, Tolerance);
+	}
+
+	return true;
+}
+
+// ============================================================================
+// Test 2: step-function gradient — validates blur actually spreads
+// ============================================================================
+//
+// Input: left half (cols 0-31) = black (0,0,0,0), right half (cols 32-63) = white (1,1,1,1).
+// After blurring with radius 4, we assert three properties without recreating the kernel:
+//   1. Far-field: col 0 stays 0, col 63 stays 1  (blur does not reach 31 pixels away).
+//   2. Boundary: col 31 and col 32 are strictly blurred — neither 0 nor 1.
+//   3. Symmetry: value(col 31) + value(col 32) == 1  (gray at the step edge).
+
+IMPLEMENT_SIMPLE_AUTOMATION_TEST(FCompositeBlurStepPassTest, "Composite.UnitTests.BlurPassStep",
+	EAutomationTestFlags_ApplicationContextMask | EAutomationTestFlags::EngineFilter | EAutomationTestFlags::NonNullRHI)
+
+bool FCompositeBlurStepPassTest::RunTest(const FString& Parameter)
+{
+	using namespace UE::Composite::Private;
+	using namespace UE::CompositeCore;
+
+	// Read back 4 columns from a single blurred step-function texture.
+	// Index: 0=far-field black (col 0), 1=boundary black (col 31),
+	//        2=boundary white (col 32), 3=far-field white (col 63).
+	static constexpr int32 ReadCols[4] = { 0, 31, 32, 63 };
+	static constexpr int32 BlurRadius  = 4;
+
+	TArray<FLinearColor> ResultRDG;
+	ResultRDG.SetNumZeroed(4);
+
+	FCompositePassBlurProxy BlurProxy(GetDefaultInputDeclArray());
+	BlurProxy.Radius    = BlurRadius;
+	BlurProxy.bAlphaOnly = false;
+
+	FCompositeTestView TestView;
+	DispatchAndWait(TestView.Get(),
+		[&ResultRDG, Proxy = MoveTemp(BlurProxy)](FRDGBuilder& GraphBuilder, const FSceneView& View) mutable
+		{
+			const FIntRect PixelRect(0, 0, GTestViewSize.X, GTestViewSize.Y);
+			const int32 CenterRow = GTestViewSize.Y / 2;
+
+			const FRDGTextureDesc InputDesc = FRDGTextureDesc::Create2D(
+				GTestViewSize, PF_A32B32G32R32F,
+				FClearValueBinding::Black, ETextureCreateFlags::None);
+			FRDGTextureRef InputTexture = GraphBuilder.CreateTexture(InputDesc, TEXT("BlurStepTest.Input"));
+
+			const FRDGTextureDesc OutputDesc = FRDGTextureDesc::Create2D(
+				GTestViewSize, PF_A32B32G32R32F,
+				FClearValueBinding::Black, ETextureCreateFlags::RenderTargetable);
+			FRDGTextureRef OutputTexture = GraphBuilder.CreateTexture(OutputDesc, TEXT("BlurStepTest.Output"));
+
+			// Upload: left half (cols 0-31) = black, right half (cols 32-63) = white.
+			{
+				FBlurTestUploadParameters* Params = GraphBuilder.AllocParameters<FBlurTestUploadParameters>();
+				Params->InputTexture = InputTexture;
+
+				GraphBuilder.AddPass(
+					RDG_EVENT_NAME("BlurStepTest.Upload"), Params,
+					ERDGPassFlags::Readback,
+					[Params](FRHICommandListImmediate& RHICmdList)
+					{
+						uint32 Stride = 0;
+						void* Data = RHICmdList.LockTexture2D(Params->InputTexture->GetRHI(), 0, RLM_WriteOnly, Stride, false);
+						for (int32 Row = 0; Row < GTestViewSize.Y; ++Row)
+						{
+							FLinearColor* RowPtr = reinterpret_cast<FLinearColor*>(static_cast<uint8*>(Data) + Row * Stride);
+							for (int32 Col = 0; Col < GTestViewSize.X; ++Col)
+							{
+								RowPtr[Col] = (Col < GTestViewSize.X / 2)
+									? FLinearColor(0.0f, 0.0f, 0.0f, 0.0f)
+									: FLinearColor(1.0f, 1.0f, 1.0f, 1.0f);
+							}
+						}
+						RHICmdList.UnlockTexture2D(Params->InputTexture->GetRHI(), 0, false);
+					});
+			}
+
+			FPassInputArray Inputs;
+			{
+				FPassInput& In = Inputs.GetArray().AddDefaulted_GetRef();
+				In.Texture         = FScreenPassTexture{ InputTexture, PixelRect };
+				In.Metadata.Filter = SF_Bilinear;
+			}
+			Inputs.OverrideOutput = FScreenPassRenderTarget{
+				FScreenPassTexture{ OutputTexture, PixelRect }, ERenderTargetLoadAction::ENoAction };
+
+			FPassContext PassContext;
+			PassContext.OutputViewRect = PixelRect;
+
+			Proxy.Add(GraphBuilder, View, Inputs, PassContext);
+
+			for (int32 i = 0; i < 4; ++i)
+			{
+				FBlurTestDownloadParameters* Params = GraphBuilder.AllocParameters<FBlurTestDownloadParameters>();
+				Params->OutputTexture = OutputTexture;
+
+				FLinearColor* ResultPtr = &ResultRDG[i];
+				const int32 Col = ReadCols[i];
+				GraphBuilder.AddPass(
+					RDG_EVENT_NAME("BlurStepTest.Download[%d]", i), Params,
+					ERDGPassFlags::Readback,
+					[ResultPtr, Params, CenterRow, Col](FRHICommandListImmediate& RHICmdList)
+					{
+						uint32 Stride = 0;
+						const void* Data = RHICmdList.LockTexture2D(Params->OutputTexture->GetRHI(), 0, RLM_ReadOnly, Stride, false);
+						const FLinearColor* RowPtr = reinterpret_cast<const FLinearColor*>(
+							static_cast<const uint8*>(Data) + CenterRow * Stride);
+						FMemory::Memcpy(ResultPtr, &RowPtr[Col], sizeof(FLinearColor));
+						RHICmdList.UnlockTexture2D(Params->OutputTexture->GetRHI(), 0, false);
+					});
+			}
+		});
+
+	// Tolerance accounts for FP16 intermediate and bilinear hardware precision.
+	constexpr float Tol = 2e-3f;
+
+	const FLinearColor& BlackFarField = ResultRDG[0]; // col  0 — 31 pixels from boundary
+	const FLinearColor& BlackBoundary = ResultRDG[1]; // col 31 — first black pixel at edge
+	const FLinearColor& WhiteBoundary = ResultRDG[2]; // col 32 — first white pixel at edge
+	const FLinearColor& WhiteFarField = ResultRDG[3]; // col 63 — 31 pixels from boundary
+
+	// Far-field pixels must be unaffected (blur radius 4 cannot reach 31 pixels away).
+	UTEST_TRUE("BlackFarField is black", BlackFarField.R < Tol);
+	UTEST_TRUE("WhiteFarField is white", WhiteFarField.R > 1.0f - Tol);
+
+	// Boundary pixels must be blurred: strictly between 0 and 1.
+	UTEST_TRUE("BlackBoundary is non-zero", BlackBoundary.R > Tol);
+	UTEST_TRUE("BlackBoundary is non-one",  BlackBoundary.R < 1.0f - Tol);
+	UTEST_TRUE("WhiteBoundary is non-zero", WhiteBoundary.R > Tol);
+	UTEST_TRUE("WhiteBoundary is non-one",  WhiteBoundary.R < 1.0f - Tol);
+
+	// Symmetry: col 31 and col 32 are mirror images of the step → they sum to 1.
+	UTEST_EQUAL_TOLERANCE("BoundarySum", BlackBoundary.R + WhiteBoundary.R, 1.0f, Tol);
+
+	return true;
+}
+
+// ============================================================================
+// Test 3: alpha-only flag — blur must spread alpha but leave RGB unchanged
+// ============================================================================
+//
+// Input: same step-function (left half = 0000, right half = 1111).
+// With bAlphaOnly=true and radius 4, boundary pixels must have:
+//   - RGB unchanged (step values — pass-through, not blurred).
+//   - Alpha blurred (strictly between 0 and 1).
+// This distinguishes "alpha-only blur" from "full blur" at the shader level.
+
+IMPLEMENT_SIMPLE_AUTOMATION_TEST(FCompositeBlurAlphaOnlyPassTest, "Composite.UnitTests.BlurPassAlphaOnly",
+	EAutomationTestFlags_ApplicationContextMask | EAutomationTestFlags::EngineFilter | EAutomationTestFlags::NonNullRHI)
+
+bool FCompositeBlurAlphaOnlyPassTest::RunTest(const FString& Parameter)
+{
+	using namespace UE::Composite::Private;
+	using namespace UE::CompositeCore;
+
+	static constexpr int32 ReadCols[4] = { 0, 31, 32, 63 };
+	static constexpr int32 BlurRadius  = 4;
+
+	TArray<FLinearColor> ResultRDG;
+	ResultRDG.SetNumZeroed(4);
+
+	FCompositePassBlurProxy BlurProxy(GetDefaultInputDeclArray());
+	BlurProxy.Radius     = BlurRadius;
+	BlurProxy.bAlphaOnly = true;
+
+	FCompositeTestView TestView;
+	DispatchAndWait(TestView.Get(),
+		[&ResultRDG, Proxy = MoveTemp(BlurProxy)](FRDGBuilder& GraphBuilder, const FSceneView& View) mutable
+		{
+			const FIntRect PixelRect(0, 0, GTestViewSize.X, GTestViewSize.Y);
+			const int32 CenterRow = GTestViewSize.Y / 2;
+
+			const FRDGTextureDesc InputDesc = FRDGTextureDesc::Create2D(
+				GTestViewSize, PF_A32B32G32R32F,
+				FClearValueBinding::Black, ETextureCreateFlags::None);
+			FRDGTextureRef InputTexture = GraphBuilder.CreateTexture(InputDesc, TEXT("BlurAlphaOnlyTest.Input"));
+
+			const FRDGTextureDesc OutputDesc = FRDGTextureDesc::Create2D(
+				GTestViewSize, PF_A32B32G32R32F,
+				FClearValueBinding::Black, ETextureCreateFlags::RenderTargetable);
+			FRDGTextureRef OutputTexture = GraphBuilder.CreateTexture(OutputDesc, TEXT("BlurAlphaOnlyTest.Output"));
+
+			// Upload: left half (cols 0-31) = black (0,0,0,0), right half (cols 32-63) = white (1,1,1,1).
+			{
+				FBlurTestUploadParameters* Params = GraphBuilder.AllocParameters<FBlurTestUploadParameters>();
+				Params->InputTexture = InputTexture;
+
+				GraphBuilder.AddPass(
+					RDG_EVENT_NAME("BlurAlphaOnlyTest.Upload"), Params,
+					ERDGPassFlags::Readback,
+					[Params](FRHICommandListImmediate& RHICmdList)
+					{
+						uint32 Stride = 0;
+						void* Data = RHICmdList.LockTexture2D(Params->InputTexture->GetRHI(), 0, RLM_WriteOnly, Stride, false);
+						for (int32 Row = 0; Row < GTestViewSize.Y; ++Row)
+						{
+							FLinearColor* RowPtr = reinterpret_cast<FLinearColor*>(static_cast<uint8*>(Data) + Row * Stride);
+							for (int32 Col = 0; Col < GTestViewSize.X; ++Col)
+							{
+								RowPtr[Col] = (Col < GTestViewSize.X / 2)
+									? FLinearColor(0.0f, 0.0f, 0.0f, 0.0f)
+									: FLinearColor(1.0f, 1.0f, 1.0f, 1.0f);
+							}
+						}
+						RHICmdList.UnlockTexture2D(Params->InputTexture->GetRHI(), 0, false);
+					});
+			}
+
+			FPassInputArray Inputs;
+			{
+				FPassInput& In = Inputs.GetArray().AddDefaulted_GetRef();
+				In.Texture         = FScreenPassTexture{ InputTexture, PixelRect };
+				In.Metadata.Filter = SF_Bilinear;
+			}
+			Inputs.OverrideOutput = FScreenPassRenderTarget{
+				FScreenPassTexture{ OutputTexture, PixelRect }, ERenderTargetLoadAction::ENoAction };
+
+			FPassContext PassContext;
+			PassContext.OutputViewRect = PixelRect;
+
+			Proxy.Add(GraphBuilder, View, Inputs, PassContext);
+
+			for (int32 i = 0; i < 4; ++i)
+			{
+				FBlurTestDownloadParameters* Params = GraphBuilder.AllocParameters<FBlurTestDownloadParameters>();
+				Params->OutputTexture = OutputTexture;
+
+				FLinearColor* ResultPtr = &ResultRDG[i];
+				const int32 Col = ReadCols[i];
+				GraphBuilder.AddPass(
+					RDG_EVENT_NAME("BlurAlphaOnlyTest.Download[%d]", i), Params,
+					ERDGPassFlags::Readback,
+					[ResultPtr, Params, CenterRow, Col](FRHICommandListImmediate& RHICmdList)
+					{
+						uint32 Stride = 0;
+						const void* Data = RHICmdList.LockTexture2D(Params->OutputTexture->GetRHI(), 0, RLM_ReadOnly, Stride, false);
+						const FLinearColor* RowPtr = reinterpret_cast<const FLinearColor*>(
+							static_cast<const uint8*>(Data) + CenterRow * Stride);
+						FMemory::Memcpy(ResultPtr, &RowPtr[Col], sizeof(FLinearColor));
+						RHICmdList.UnlockTexture2D(Params->OutputTexture->GetRHI(), 0, false);
+					});
+			}
+		});
+
+	// Tolerance accounts for FP16 intermediate and bilinear hardware precision.
+	constexpr float Tol = 2e-3f;
+
+	const FLinearColor& BlackFarField = ResultRDG[0]; // col  0 — 31 pixels from boundary
+	const FLinearColor& BlackBoundary = ResultRDG[1]; // col 31 — first black pixel at edge
+	const FLinearColor& WhiteBoundary = ResultRDG[2]; // col 32 — first white pixel at edge
+	const FLinearColor& WhiteFarField = ResultRDG[3]; // col 63 — 31 pixels from boundary
+
+	// Far-field: neither RGB nor alpha is affected at radius 4 / 31-pixel distance.
+	UTEST_TRUE("BlackFarField.R is black", BlackFarField.R < Tol);
+	UTEST_TRUE("BlackFarField.A is zero",  BlackFarField.A < Tol);
+	UTEST_TRUE("WhiteFarField.R is white", WhiteFarField.R > 1.0f - Tol);
+	UTEST_TRUE("WhiteFarField.A is one",   WhiteFarField.A > 1.0f - Tol);
+
+	// Boundary RGB: pass-through — must NOT be blurred.
+	UTEST_TRUE("BlackBoundary.R is not blurred", BlackBoundary.R < Tol);
+	UTEST_TRUE("WhiteBoundary.R is not blurred", WhiteBoundary.R > 1.0f - Tol);
+
+	// Boundary alpha: must BE blurred (strictly between 0 and 1).
+	UTEST_TRUE("BlackBoundary.A is non-zero",  BlackBoundary.A > Tol);
+	UTEST_TRUE("BlackBoundary.A is non-one",   BlackBoundary.A < 1.0f - Tol);
+	UTEST_TRUE("WhiteBoundary.A is non-zero",  WhiteBoundary.A > Tol);
+	UTEST_TRUE("WhiteBoundary.A is non-one",   WhiteBoundary.A < 1.0f - Tol);
+
+	return true;
+}
+
+#endif // WITH_DEV_AUTOMATION_TESTS
+
+} // namespace UE::Composite::Tests
diff --git a/Engine/Plugins/Compositing/Composite/Source/Composite/Private/Tests/CompositePassCenteredScaleTest.cpp b/Engine/Plugins/Compositing/Composite/Source/Composite/Private/Tests/CompositePassCenteredScaleTest.cpp
new file mode 100644
--- /dev/null
+++ b/Engine/Plugins/Compositing/Composite/Source/Composite/Private/Tests/CompositePassCenteredScaleTest.cpp
@@ -0,0 +1,184 @@
+// Copyright Epic Games, Inc. All Rights Reserved.
+
+#include "CoreTypes.h"
+#include "RenderGraphBuilder.h"
+#include "RHI.h"
+#include "ScreenPass.h"
+
+#include "Passes/CompositePassCenteredScale.h"
+#include "Tests/CompositeTestHelpers.h"
+
+#if WITH_DEV_AUTOMATION_TESTS
+#include "Misc/AutomationTest.h"
+#endif
+
+namespace UE::Composite::Tests
+{
+
+#if WITH_DEV_AUTOMATION_TESTS
+
+// -- Upload/download parameter structs --
+
+BEGIN_SHADER_PARAMETER_STRUCT(FCenteredScaleTestUploadParameters, )
+	RDG_TEXTURE_ACCESS(InputTexture, ERHIAccess::CopyDest)
+END_SHADER_PARAMETER_STRUCT()
+
+BEGIN_SHADER_PARAMETER_STRUCT(FCenteredScaleTestDownloadParameters, )
+	RDG_TEXTURE_ACCESS(OutputTexture, ERHIAccess::CopySrc)
+END_SHADER_PARAMETER_STRUCT()
+
+namespace CenteredScale  // unique name prevents symbol collision in unity builds
+{
+
+// -- CPU reference --
+
+// Input: horizontal gradient (Red = col/(W-1), G=B=0, A=1).
+// Bilinear-sample the output pixel at SampleCol after a centered scale of ScaleX.
+// Formula: u_in = (u_out - 0.5)*ScaleX + 0.5, where u_out = (SampleCol+0.5)/W.
+// Bilinear on a linear gradient: red = clamp(u_in*W - 0.5, 0, W-1) / (W-1).
+static float CenteredScaleExpectedR(float ScaleX, int32 SampleCol, int32 W)
+{
+	const float UOut = (SampleCol + 0.5f) / W;
+	const float UIn  = (UOut - 0.5f) * ScaleX + 0.5f;
+	return FMath::Clamp(UIn * W - 0.5f, 0.0f, float(W - 1)) / (W - 1);
+}
+
+// -- Test scenarios --
+
+struct FCenteredScaleScenario
+{
+	const TCHAR* Label;
+	FVector2f    ScaleUV;
+};
+
+// Input: horizontal gradient (Red = col/(W-1), G=B=0, A=1).
+// Sampled at column W/4 - offset from center so the scale factor shifts the sample point.
+static const FCenteredScaleScenario GScenarios[] =
+{
+	{ TEXT("Identity"),       FVector2f(1.0f,  1.0f)  },
+	{ TEXT("ZoomIn"),         FVector2f(0.5f,  0.5f)  },
+	{ TEXT("ZoomOut"),        FVector2f(2.0f,  2.0f)  },
+	{ TEXT("HorizontalOnly"), FVector2f(0.75f, 1.0f)  },
+};
+
+}  // namespace CenteredScale
+
+IMPLEMENT_SIMPLE_AUTOMATION_TEST(FCompositeCenteredScalePassTest, "Composite.UnitTests.CenteredScalePass",
+	EAutomationTestFlags_ApplicationContextMask | EAutomationTestFlags::EngineFilter | EAutomationTestFlags::NonNullRHI)
+
+bool FCompositeCenteredScalePassTest::RunTest(const FString& Parameter)
+{
+	using namespace UE::Composite::Private;
+	using namespace UE::CompositeCore;
+	using namespace CenteredScale;
+
+	const int32 NumScenarios = UE_ARRAY_COUNT(GScenarios);
+
+	TArray<FLinearColor> ResultRDG;
+	ResultRDG.SetNumZeroed(NumScenarios);
+
+	TArray<FCenteredScalePassProxy> Proxies;
+	Proxies.Reserve(NumScenarios);
+	for (int32 i = 0; i < NumScenarios; ++i)
+	{
+		const FCenteredScaleScenario& S = GScenarios[i];
+		FCenteredScalePassProxy Proxy(GetDefaultInputDeclArray());
+		Proxy.ScaleUV = S.ScaleUV;
+		Proxies.Emplace(MoveTemp(Proxy));
+	}
+
+	FCompositeTestView TestView;
+	DispatchAndWait(TestView.Get(),
+		[&ResultRDG, NumScenarios, Proxies = MoveTemp(Proxies)](FRDGBuilder& GraphBuilder, const FSceneView& View) mutable
+		{
+			const FIntRect PixelRect(0, 0, GTestViewSize.X, GTestViewSize.Y);
+
+			for (int32 i = 0; i < NumScenarios; ++i)
+			{
+				const FRDGTextureDesc InputDesc = FRDGTextureDesc::Create2D(
+					GTestViewSize, PF_A32B32G32R32F,
+					FClearValueBinding::Black, ETextureCreateFlags::None);
+				FRDGTextureRef InputTexture = GraphBuilder.CreateTexture(InputDesc, TEXT("CenteredScaleTest.Input"));
+
+				const FRDGTextureDesc OutputDesc = FRDGTextureDesc::Create2D(
+					GTestViewSize, PF_A32B32G32R32F,
+					FClearValueBinding::Black, ETextureCreateFlags::RenderTargetable);
+				FRDGTextureRef OutputTexture = GraphBuilder.CreateTexture(OutputDesc, TEXT("CenteredScaleTest.Output"));
+
+				{
+					FCenteredScaleTestUploadParameters* Params = GraphBuilder.AllocParameters<FCenteredScaleTestUploadParameters>();
+					Params->InputTexture = InputTexture;
+
+					GraphBuilder.AddPass(
+						RDG_EVENT_NAME("CenteredScaleTest.Upload[%d]", i), Params,
+						ERDGPassFlags::Readback,
+						[Params](FRHICommandListImmediate& RHICmdList)
+						{
+							uint32 Stride = 0;
+							void* Data = RHICmdList.LockTexture2D(Params->InputTexture->GetRHI(), 0, RLM_WriteOnly, Stride, false);
+							for (int32 Row = 0; Row < GTestViewSize.Y; ++Row)
+							{
+								FLinearColor* RowPtr = reinterpret_cast<FLinearColor*>(static_cast<uint8*>(Data) + Row * Stride);
+								for (int32 Col = 0; Col < GTestViewSize.X; ++Col)
+								{
+									RowPtr[Col] = FLinearColor(Col / (GTestViewSize.X - 1.0f), 0.0f, 0.0f, 1.0f);
+								}
+							}
+							RHICmdList.UnlockTexture2D(Params->InputTexture->GetRHI(), 0, false);
+						});
+				}
+
+				FPassInputArray Inputs;
+				{
+					FPassInput& In = Inputs.GetArray().AddDefaulted_GetRef();
+					In.Texture         = FScreenPassTexture{ InputTexture, PixelRect };
+					In.Metadata.Filter = SF_Bilinear;
+				}
+				Inputs.OverrideOutput = FScreenPassRenderTarget{
+					FScreenPassTexture{ OutputTexture, PixelRect }, ERenderTargetLoadAction::ENoAction };
+
+				FPassContext PassContext;
+				PassContext.OutputViewRect = PixelRect;
+
+				Proxies[i].Add(GraphBuilder, View, Inputs, PassContext);
+
+				{
+					FCenteredScaleTestDownloadParameters* Params = GraphBuilder.AllocParameters<FCenteredScaleTestDownloadParameters>();
+					Params->OutputTexture = OutputTexture;
+
+					FLinearColor* ResultPtr = &ResultRDG[i];
+					GraphBuilder.AddPass(
+						RDG_EVENT_NAME("CenteredScaleTest.Download[%d]", i), Params,
+						ERDGPassFlags::Readback,
+						[ResultPtr, Params](FRHICommandListImmediate& RHICmdList)
+						{
+							uint32 Stride = 0;
+							const void* Data = RHICmdList.LockTexture2D(Params->OutputTexture->GetRHI(), 0, RLM_ReadOnly, Stride, false);
+							const FLinearColor* RowPtr = reinterpret_cast<const FLinearColor*>(
+								static_cast<const uint8*>(Data) + (GTestViewSize.Y / 2) * Stride);
+							FMemory::Memcpy(ResultPtr, &RowPtr[GTestViewSize.X / 4], sizeof(FLinearColor));
+							RHICmdList.UnlockTexture2D(Params->OutputTexture->GetRHI(), 0, false);
+						});
+				}
+			}
+		});
+
+	constexpr float Tolerance = 1e-4f;
+	const int32 SampleCol = GTestViewSize.X / 4;
+	for (int32 i = 0; i < NumScenarios; ++i)
+	{
+		const FCenteredScaleScenario& S  = GScenarios[i];
+		const FLinearColor&           Actual = ResultRDG[i];
+
+		UTEST_EQUAL_TOLERANCE(S.Label, Actual.R, CenteredScaleExpectedR(S.ScaleUV.X, SampleCol, GTestViewSize.X), Tolerance);
+		UTEST_EQUAL_TOLERANCE(S.Label, Actual.G, 0.0f, Tolerance);
+		UTEST_EQUAL_TOLERANCE(S.Label, Actual.B, 0.0f, Tolerance);
+		UTEST_EQUAL_TOLERANCE(S.Label, Actual.A, 1.0f, Tolerance);
+	}
+
+	return true;
+}
+
+#endif // WITH_DEV_AUTOMATION_TESTS
+
+} // namespace UE::Composite::Tests
diff --git a/Engine/Plugins/Compositing/Composite/Source/Composite/Private/Tests/CompositePassColorGradingTest.cpp b/Engine/Plugins/Compositing/Composite/Source/Composite/Private/Tests/CompositePassColorGradingTest.cpp
new file mode 100644
--- /dev/null
+++ b/Engine/Plugins/Compositing/Composite/Source/Composite/Private/Tests/CompositePassColorGradingTest.cpp
@@ -0,0 +1,212 @@
+// Copyright Epic Games, Inc. All Rights Reserved.
+
+#include "CoreTypes.h"
+#include "RenderGraphBuilder.h"
+#include "RHI.h"
+#include "ScreenPass.h"
+
+#include "Passes/CompositePassColorGrading.h"
+#include "Tests/CompositeTestHelpers.h"
+
+#if WITH_DEV_AUTOMATION_TESTS
+#include "Misc/AutomationTest.h"
+#endif
+
+namespace UE::Composite::Tests
+{
+
+#if WITH_DEV_AUTOMATION_TESTS
+
+// -- Upload/download parameter structs --
+
+BEGIN_SHADER_PARAMETER_STRUCT(FColorGradingTestUploadParameters, )
+	RDG_TEXTURE_ACCESS(InputTexture, ERHIAccess::CopyDest)
+END_SHADER_PARAMETER_STRUCT()
+
+BEGIN_SHADER_PARAMETER_STRUCT(FColorGradingTestDownloadParameters, )
+	RDG_TEXTURE_ACCESS(OutputTexture, ERHIAccess::CopySrc)
+END_SHADER_PARAMETER_STRUCT()
+
+namespace ColorGrading  // unique name prevents symbol collision in unity builds
+{
+
+// -- CPU references --
+
+// Default FColorGradingSettings is identity: output == input.
+static FLinearColor ColorGradingIdentityCPU(FLinearColor In)
+{
+	return In;
+}
+
+// Black input (0,0,0) is fully in the shadow range (luma=0, shadow weight=1).
+// Global Offset.w adds uniformly to all channels in shadow CC, so output = (Offset, Offset, Offset, A).
+static FLinearColor ColorGradingShadowOffsetCPU(FLinearColor In, float Offset)
+{
+	return FLinearColor(In.R + Offset, In.G + Offset, In.B + Offset, In.A);
+}
+
+// -- Test scenarios --
+
+struct FColorGradingScenario
+{
+	const TCHAR* Label;
+	FLinearColor  Input;
+	bool          bUnpremult;
+	float         GlobalOffsetW; // Global.Offset.w — 0 means no offset
+};
+
+static const FColorGradingScenario GScenarios[] =
+{
+	// Identity: default settings leave input unchanged regardless of range blending weights.
+	{ TEXT("IdentityMidtone"),    { 0.5f, 0.5f, 0.5f, 1.0f }, false, 0.0f },
+	{ TEXT("IdentityColorful"),   { 0.6f, 0.3f, 0.1f, 1.0f }, false, 0.0f },
+	{ TEXT("IdentityHighlight"),  { 0.8f, 0.8f, 0.8f, 1.0f }, false, 0.0f },
+
+	// Unpremult + identity grading: unpremult and premult cancel → output == input (alpha > 0).
+	{ TEXT("UnpremultRoundtrip"), { 0.4f, 0.3f, 0.2f, 0.5f }, true,  0.0f },
+
+	// Black input lies entirely in the shadow range (luma=0, shadow weight=1).
+	// Offset.w=0.2 adds 0.2 to all channels via the shadow CC term.
+	{ TEXT("ShadowOffset"),       { 0.0f, 0.0f, 0.0f, 1.0f }, false, 0.2f },
+};
+
+// CPU reference dispatcher.
+static FLinearColor ColorGradingCPU(const FColorGradingScenario& S)
+{
+	if (S.GlobalOffsetW != 0.0f)
+	{
+		return ColorGradingShadowOffsetCPU(S.Input, S.GlobalOffsetW);
+	}
+	return ColorGradingIdentityCPU(S.Input);
+}
+
+}  // namespace ColorGrading
+
+// -- Test --
+
+IMPLEMENT_SIMPLE_AUTOMATION_TEST(FCompositeColorGradingPassTest, "Composite.UnitTests.ColorGradingPass",
+	EAutomationTestFlags_ApplicationContextMask | EAutomationTestFlags::EngineFilter | EAutomationTestFlags::NonNullRHI)
+
+bool FCompositeColorGradingPassTest::RunTest(const FString& Parameter)
+{
+	using namespace UE::Composite::Private;
+	using namespace UE::CompositeCore;
+	using namespace ColorGrading;
+
+	const int32 NumScenarios = UE_ARRAY_COUNT(GScenarios);
+
+	TArray<FLinearColor> ResultRDG;
+	ResultRDG.SetNumZeroed(NumScenarios);
+
+	TArray<FColorGradingCompositePassProxy> Proxies;
+	Proxies.Reserve(NumScenarios);
+	for (int32 i = 0; i < NumScenarios; ++i)
+	{
+		const FColorGradingScenario& S = GScenarios[i];
+		FColorGradingCompositePassProxy Proxy(GetDefaultInputDeclArray());
+		Proxy.bUnpremult = S.bUnpremult;
+		if (S.GlobalOffsetW != 0.0f)
+		{
+			Proxy.ColorGradingSettings.Global.Offset = FVector4(0.0, 0.0, 0.0, static_cast<double>(S.GlobalOffsetW));
+		}
+		Proxies.Emplace(MoveTemp(Proxy));
+	}
+
+	FCompositeTestView TestView;
+	DispatchAndWait(TestView.Get(),
+		[&ResultRDG, NumScenarios, Proxies = MoveTemp(Proxies)](FRDGBuilder& GraphBuilder, const FSceneView& View) mutable
+		{
+			const FIntRect PixelRect(0, 0, GTestViewSize.X, GTestViewSize.Y);
+
+			for (int32 i = 0; i < NumScenarios; ++i)
+			{
+				const FColorGradingScenario& S = GScenarios[i];
+
+				const FRDGTextureDesc InputDesc = FRDGTextureDesc::Create2D(
+					GTestViewSize, PF_A32B32G32R32F,
+					FClearValueBinding::Black, ETextureCreateFlags::None);
+				FRDGTextureRef InputTexture = GraphBuilder.CreateTexture(InputDesc, TEXT("ColorGradingTest.Input"));
+
+				const FRDGTextureDesc OutputDesc = FRDGTextureDesc::Create2D(
+					GTestViewSize, PF_A32B32G32R32F,
+					FClearValueBinding::Black, ETextureCreateFlags::RenderTargetable);
+				FRDGTextureRef OutputTexture = GraphBuilder.CreateTexture(OutputDesc, TEXT("ColorGradingTest.Output"));
+
+				{
+					FColorGradingTestUploadParameters* Params = GraphBuilder.AllocParameters<FColorGradingTestUploadParameters>();
+					Params->InputTexture = InputTexture;
+
+					const FLinearColor PixelValue = S.Input;
+					GraphBuilder.AddPass(
+						RDG_EVENT_NAME("ColorGradingTest.Upload[%d]", i), Params,
+						ERDGPassFlags::Readback,
+						[PixelValue, Params](FRHICommandListImmediate& RHICmdList)
+						{
+							uint32 Stride = 0;
+							void* Data = RHICmdList.LockTexture2D(Params->InputTexture->GetRHI(), 0, RLM_WriteOnly, Stride, false);
+							for (int32 Row = 0; Row < GTestViewSize.Y; ++Row)
+							{
+								FLinearColor* RowPtr = reinterpret_cast<FLinearColor*>(static_cast<uint8*>(Data) + Row * Stride);
+								for (int32 Col = 0; Col < GTestViewSize.X; ++Col)
+								{
+									RowPtr[Col] = PixelValue;
+								}
+							}
+							RHICmdList.UnlockTexture2D(Params->InputTexture->GetRHI(), 0, false);
+						});
+				}
+
+				FPassInputArray Inputs;
+				{
+					FPassInput& In = Inputs.GetArray().AddDefaulted_GetRef();
+					In.Texture         = FScreenPassTexture{ InputTexture, PixelRect };
+					In.Metadata.Filter = SF_Point;
+				}
+				Inputs.OverrideOutput = FScreenPassRenderTarget{
+					FScreenPassTexture{ OutputTexture, PixelRect }, ERenderTargetLoadAction::ENoAction };
+
+				FPassContext PassContext;
+				PassContext.OutputViewRect = PixelRect;
+
+				Proxies[i].Add(GraphBuilder, View, Inputs, PassContext);
+
+				{
+					FColorGradingTestDownloadParameters* Params = GraphBuilder.AllocParameters<FColorGradingTestDownloadParameters>();
+					Params->OutputTexture = OutputTexture;
+
+					FLinearColor* ResultPtr = &ResultRDG[i];
+					GraphBuilder.AddPass(
+						RDG_EVENT_NAME("ColorGradingTest.Download[%d]", i), Params,
+						ERDGPassFlags::Readback,
+						[ResultPtr, Params](FRHICommandListImmediate& RHICmdList)
+						{
+							uint32 Stride = 0;
+							const void* Data = RHICmdList.LockTexture2D(Params->OutputTexture->GetRHI(), 0, RLM_ReadOnly, Stride, false);
+							const FLinearColor* RowPtr = reinterpret_cast<const FLinearColor*>(
+								static_cast<const uint8*>(Data) + (GTestViewSize.Y / 2) * Stride);
+							FMemory::Memcpy(ResultPtr, &RowPtr[GTestViewSize.X / 2], sizeof(FLinearColor));
+							RHICmdList.UnlockTexture2D(Params->OutputTexture->GetRHI(), 0, false);
+						});
+				}
+			}
+		});
+
+	constexpr float Tolerance = 1e-4f;
+	for (int32 i = 0; i < NumScenarios; ++i)
+	{
+		const FColorGradingScenario& S = GScenarios[i];
+		const FLinearColor Expected = ColorGradingCPU(S);
+		const FLinearColor& Actual  = ResultRDG[i];
+
+		UTEST_EQUAL_TOLERANCE(S.Label, Actual.R, Expected.R, Tolerance);
+		UTEST_EQUAL_TOLERANCE(S.Label, Actual.G, Expected.G, Tolerance);
+		UTEST_EQUAL_TOLERANCE(S.Label, Actual.B, Expected.B, Tolerance);
+		UTEST_EQUAL_TOLERANCE(S.Label, Actual.A, Expected.A, Tolerance);
+	}
+
+	return true;
+}
+
+#endif // WITH_DEV_AUTOMATION_TESTS
+
+} // namespace UE::Composite::Tests
diff --git a/Engine/Plugins/Compositing/Composite/Source/Composite/Private/Tests/CompositePassColorKeyerTest.cpp b/Engine/Plugins/Compositing/Composite/Source/Composite/Private/Tests/CompositePassColorKeyerTest.cpp
new file mode 100644
--- /dev/null
+++ b/Engine/Plugins/Compositing/Composite/Source/Composite/Private/Tests/CompositePassColorKeyerTest.cpp
@@ -0,0 +1,189 @@
+// Copyright Epic Games, Inc. All Rights Reserved.
+
+#include "CoreTypes.h"
+#include "RenderGraphBuilder.h"
+#include "RHI.h"
+#include "ScreenPass.h"
+
+#include "Passes/CompositePassColorKeyer.h"
+#include "Tests/CompositeTestHelpers.h"
+
+#if WITH_DEV_AUTOMATION_TESTS
+#include "Misc/AutomationTest.h"
+#endif
+
+namespace UE::Composite::Tests
+{
+
+#if WITH_DEV_AUTOMATION_TESTS
+
+// -- Upload/download parameter structs --
+
+BEGIN_SHADER_PARAMETER_STRUCT(FColorKeyerTestUploadParameters, )
+	RDG_TEXTURE_ACCESS(InputTexture, ERHIAccess::CopyDest)
+END_SHADER_PARAMETER_STRUCT()
+
+BEGIN_SHADER_PARAMETER_STRUCT(FColorKeyerTestDownloadParameters, )
+	RDG_TEXTURE_ACCESS(OutputTexture, ERHIAccess::CopySrc)
+END_SHADER_PARAMETER_STRUCT()
+
+namespace ColorKeyer  // unique name prevents symbol collision in unity builds
+{
+
+struct FColorKeyerScenario
+{
+	const TCHAR* Label;
+	FLinearColor  Input;
+	FLinearColor  KeyColor;
+	float         RedWeight;
+	float         BlueWeight;
+	float         AlphaMin;
+	float         AlphaMax;
+	bool          bInvertAlpha;
+	FLinearColor  Expected;  // premultiplied output
+};
+
+static const FColorKeyerScenario GScenarios[] =
+{
+	// Pixel identical to the key color → fully keyed out. Alpha=0, RGB premultiplied to zero.
+	{ TEXT("GreenOnGreen"),   { 0.0f, 1.0f, 0.0f, 1.0f }, { 0.0f, 1.0f, 0.0f, 1.0f }, 0.5f, 0.5f, 0.0f, 1.0f, false, { 0.0f, 0.0f, 0.0f, 0.0f } },
+	// Pixel has no key-channel content → fully preserved. Alpha=1, RGB premultiplied unchanged.
+	{ TEXT("RedOnGreen"),     { 1.0f, 0.0f, 0.0f, 1.0f }, { 0.0f, 1.0f, 0.0f, 1.0f }, 0.5f, 0.5f, 0.0f, 1.0f, false, { 1.0f, 0.0f, 0.0f, 1.0f } },
+	// Inverted alpha: fully-keyed pixel gets alpha=1. Fill is zero, so premultiplied RGB stays zero.
+	{ TEXT("GreenInverted"),  { 0.0f, 1.0f, 0.0f, 1.0f }, { 0.0f, 1.0f, 0.0f, 1.0f }, 0.5f, 0.5f, 0.0f, 1.0f, true,  { 0.0f, 0.0f, 0.0f, 1.0f } },
+};
+
+}  // namespace ColorKeyer
+
+IMPLEMENT_SIMPLE_AUTOMATION_TEST(FCompositeColorKeyerPassTest, "Composite.UnitTests.ColorKeyerPass",
+	EAutomationTestFlags_ApplicationContextMask | EAutomationTestFlags::EngineFilter | EAutomationTestFlags::NonNullRHI)
+
+bool FCompositeColorKeyerPassTest::RunTest(const FString& Parameter)
+{
+	using namespace UE::Composite::Private;
+	using namespace UE::CompositeCore;
+	using namespace ColorKeyer;
+
+	const int32 NumScenarios = UE_ARRAY_COUNT(GScenarios);
+
+	TArray<FLinearColor> ResultRDG;
+	ResultRDG.SetNumZeroed(NumScenarios);
+
+	TArray<FCompositePassColorKeyerProxy> Proxies;
+	Proxies.Reserve(NumScenarios);
+	for (int32 i = 0; i < NumScenarios; ++i)
+	{
+		const FColorKeyerScenario& S = GScenarios[i];
+		FCompositePassColorKeyerProxy Proxy(GetDefaultInputDeclArray());
+		Proxy.ScreenType              = ECompositeColorKeyerScreenType::Green;
+		Proxy.KeyColor                = S.KeyColor;
+		Proxy.Params0                 = FVector4f(S.RedWeight, 0.5f, S.BlueWeight, 0.0f);
+		Proxy.Params1                 = FVector4f(S.AlphaMin, S.AlphaMax, 0.0f, 0.0f);
+		Proxy.Visualization           = ECompositeColorKeyerVisualization::Key;
+		Proxy.bPreserveVignetteAfterKey = false;
+		Proxy.DenoiseMethod           = ECompositeDenoiseMethod::None;
+		Proxy.bDespillOnly            = false;
+		Proxy.bInvertAlpha            = S.bInvertAlpha;
+		Proxies.Emplace(MoveTemp(Proxy));
+	}
+
+	FCompositeTestView TestView;
+	DispatchAndWait(TestView.Get(),
+		[&ResultRDG, NumScenarios, Proxies = MoveTemp(Proxies)](FRDGBuilder& GraphBuilder, const FSceneView& View) mutable
+		{
+			const FIntRect PixelRect(0, 0, GTestViewSize.X, GTestViewSize.Y);
+
+			for (int32 i = 0; i < NumScenarios; ++i)
+			{
+				const FColorKeyerScenario& S = GScenarios[i];
+
+				const FRDGTextureDesc InputDesc = FRDGTextureDesc::Create2D(
+					GTestViewSize, PF_A32B32G32R32F,
+					FClearValueBinding::Black, ETextureCreateFlags::None);
+				FRDGTextureRef InputTexture = GraphBuilder.CreateTexture(InputDesc, TEXT("ColorKeyerTest.Input"));
+
+				// ColorKeyer writes via UAV: output needs UAV + RenderTargetable flags.
+				const FRDGTextureDesc OutputDesc = FRDGTextureDesc::Create2D(
+					GTestViewSize, PF_A32B32G32R32F,
+					FClearValueBinding::Black,
+					ETextureCreateFlags::RenderTargetable | ETextureCreateFlags::UAV);
+				FRDGTextureRef OutputTexture = GraphBuilder.CreateTexture(OutputDesc, TEXT("ColorKeyerTest.Output"));
+
+				{
+					FColorKeyerTestUploadParameters* Params = GraphBuilder.AllocParameters<FColorKeyerTestUploadParameters>();
+					Params->InputTexture = InputTexture;
+
+					const FLinearColor PixelValue = S.Input;
+					GraphBuilder.AddPass(
+						RDG_EVENT_NAME("ColorKeyerTest.Upload[%d]", i), Params,
+						ERDGPassFlags::Readback,
+						[PixelValue, Params](FRHICommandListImmediate& RHICmdList)
+						{
+							uint32 Stride = 0;
+							void* Data = RHICmdList.LockTexture2D(Params->InputTexture->GetRHI(), 0, RLM_WriteOnly, Stride, false);
+							for (int32 Row = 0; Row < GTestViewSize.Y; ++Row)
+							{
+								FLinearColor* RowPtr = reinterpret_cast<FLinearColor*>(static_cast<uint8*>(Data) + Row * Stride);
+								for (int32 Col = 0; Col < GTestViewSize.X; ++Col)
+								{
+									RowPtr[Col] = PixelValue;
+								}
+							}
+							RHICmdList.UnlockTexture2D(Params->InputTexture->GetRHI(), 0, false);
+						});
+				}
+
+				FPassInputArray Inputs;
+				{
+					FPassInput& In = Inputs.GetArray().AddDefaulted_GetRef();
+					In.Texture         = FScreenPassTexture{ InputTexture, PixelRect };
+					In.Metadata.Filter = SF_Point;
+				}
+				Inputs.OverrideOutput = FScreenPassRenderTarget{
+					FScreenPassTexture{ OutputTexture, PixelRect }, ERenderTargetLoadAction::ENoAction };
+
+				FPassContext PassContext;
+				PassContext.OutputViewRect = PixelRect;
+
+				Proxies[i].Add(GraphBuilder, View, Inputs, PassContext);
+
+				{
+					FColorKeyerTestDownloadParameters* Params = GraphBuilder.AllocParameters<FColorKeyerTestDownloadParameters>();
+					Params->OutputTexture = OutputTexture;
+
+					FLinearColor* ResultPtr = &ResultRDG[i];
+					GraphBuilder.AddPass(
+						RDG_EVENT_NAME("ColorKeyerTest.Download[%d]", i), Params,
+						ERDGPassFlags::Readback,
+						[ResultPtr, Params](FRHICommandListImmediate& RHICmdList)
+						{
+							uint32 Stride = 0;
+							const void* Data = RHICmdList.LockTexture2D(Params->OutputTexture->GetRHI(), 0, RLM_ReadOnly, Stride, false);
+							const FLinearColor* RowPtr = reinterpret_cast<const FLinearColor*>(
+								static_cast<const uint8*>(Data) + (GTestViewSize.Y / 2) * Stride);
+							FMemory::Memcpy(ResultPtr, &RowPtr[GTestViewSize.X / 2], sizeof(FLinearColor));
+							RHICmdList.UnlockTexture2D(Params->OutputTexture->GetRHI(), 0, false);
+						});
+				}
+			}
+		});
+
+	constexpr float Tolerance = 1e-4f;
+	for (int32 i = 0; i < NumScenarios; ++i)
+	{
+		const FColorKeyerScenario& S = GScenarios[i];
+		const FLinearColor& Expected = S.Expected;
+		const FLinearColor& Actual = ResultRDG[i];
+
+		UTEST_EQUAL_TOLERANCE(S.Label, Actual.R, Expected.R, Tolerance);
+		UTEST_EQUAL_TOLERANCE(S.Label, Actual.G, Expected.G, Tolerance);
+		UTEST_EQUAL_TOLERANCE(S.Label, Actual.B, Expected.B, Tolerance);
+		UTEST_EQUAL_TOLERANCE(S.Label, Actual.A, Expected.A, Tolerance);
+	}
+
+	return true;
+}
+
+#endif // WITH_DEV_AUTOMATION_TESTS
+
+} // namespace UE::Composite::Tests
diff --git a/Engine/Plugins/Compositing/Composite/Source/Composite/Private/Tests/CompositePassMaskingTest.cpp b/Engine/Plugins/Compositing/Composite/Source/Composite/Private/Tests/CompositePassMaskingTest.cpp
new file mode 100644
--- /dev/null
+++ b/Engine/Plugins/Compositing/Composite/Source/Composite/Private/Tests/CompositePassMaskingTest.cpp
@@ -0,0 +1,260 @@
+// Copyright Epic Games, Inc. All Rights Reserved.
+
+#include "CoreTypes.h"
+#include "Math/Color.h"
+#include "RenderGraphBuilder.h"
+#include "RHI.h"
+#include "ScreenPass.h"
+
+#include "Passes/CompositePassMasking.h"
+#include "Tests/CompositeTestHelpers.h"
+
+#if WITH_DEV_AUTOMATION_TESTS
+#include "Misc/AutomationTest.h"
+#endif
+
+namespace UE::Composite::Tests
+{
+
+#if WITH_DEV_AUTOMATION_TESTS
+
+// -- Upload/download parameter structs --
+
+BEGIN_SHADER_PARAMETER_STRUCT(FMaskingTestMainUploadParameters, )
+	RDG_TEXTURE_ACCESS(InputTexture, ERHIAccess::CopyDest)
+END_SHADER_PARAMETER_STRUCT()
+
+BEGIN_SHADER_PARAMETER_STRUCT(FMaskingTestMaskUploadParameters, )
+	RDG_TEXTURE_ACCESS(MaskTexture, ERHIAccess::CopyDest)
+END_SHADER_PARAMETER_STRUCT()
+
+BEGIN_SHADER_PARAMETER_STRUCT(FMaskingTestDownloadParameters, )
+	RDG_TEXTURE_ACCESS(OutputTexture, ERHIAccess::CopySrc)
+END_SHADER_PARAMETER_STRUCT()
+
+namespace Masking  // unique name prevents symbol collision in unity builds
+{
+
+// -- CPU reference --
+
+// Default metadata: bInvertedAlpha=false (InvertAlpha = identity), Encoding=Linear (Decode/Encode = identity).
+// Masking shader: AlphaMask = saturate(Mask[Channel]); optional invert; Output.A = Input.A * AlphaMask.
+// Optional unpremult (divide RGB by A) before masking and premult (multiply RGB by A) after.
+static FLinearColor MaskingCPU(FLinearColor Main, FLinearColor Mask, uint32 Channel, bool bUnpremult, bool bInvert)
+{
+	if (bUnpremult)
+	{
+		const float InvA = (Main.A > 0.0f) ? 1.0f / Main.A : 0.0f;
+		Main.R *= InvA;
+		Main.G *= InvA;
+		Main.B *= InvA;
+	}
+
+	const float MaskChannels[4] = { Mask.R, Mask.G, Mask.B, Mask.A };
+	float AlphaMask = FMath::Clamp(MaskChannels[Channel], 0.0f, 1.0f);
+	if (bInvert)
+	{
+		AlphaMask = 1.0f - AlphaMask;
+	}
+
+	Main.A *= AlphaMask;
+
+	if (bUnpremult)
+	{
+		Main.R *= Main.A;
+		Main.G *= Main.A;
+		Main.B *= Main.A;
+	}
+
+	return Main;
+}
+
+// -- Test scenarios --
+
+struct FMaskingScenario
+{
+	const TCHAR* Label;
+	FLinearColor  MainInput;
+	FLinearColor  MaskInput;
+	uint32        Channel;    // 0=R, 1=G, 2=B, 3=A
+	bool          bUnpremult;
+	bool          bInvert;
+};
+
+static const FMaskingScenario GScenarios[] =
+{
+	// Mask from red channel: AlphaMask = 1.0 → Output.A = Main.A * 1.0.
+	{ TEXT("MaskFromRed"),   { 0.5f, 0.3f, 0.7f, 0.8f }, { 1.0f, 0.0f, 0.0f, 1.0f }, 0, false, false },
+	// Mask from green channel: AlphaMask = 0.6 → Output.A = Main.A * 0.6 = 0.48.
+	{ TEXT("MaskFromGreen"), { 0.5f, 0.3f, 0.7f, 0.8f }, { 0.0f, 0.6f, 0.0f, 1.0f }, 1, false, false },
+	// Mask from alpha channel: AlphaMask = 0.4 → Output.A = Main.A * 0.4 = 0.32.
+	{ TEXT("MaskFromAlpha"), { 0.5f, 0.3f, 0.7f, 0.8f }, { 0.0f, 0.0f, 0.0f, 0.4f }, 3, false, false },
+	// Inverted mask: AlphaMask = 1.0 - 0.75 = 0.25 → Output.A = Main.A * 0.25 = 0.2.
+	{ TEXT("InvertedMask"),  { 0.5f, 0.3f, 0.7f, 0.8f }, { 0.75f, 0.0f, 0.0f, 1.0f }, 0, false, true  },
+	// Unpremult: divide by alpha, apply 0.5 mask, multiply back → RGB changes.
+	{ TEXT("Unpremult"),     { 0.4f, 0.6f, 0.2f, 0.8f }, { 0.5f, 0.0f, 0.0f, 1.0f }, 0, true,  false },
+};
+
+}  // namespace Masking
+
+IMPLEMENT_SIMPLE_AUTOMATION_TEST(FCompositeMaskingPassTest, "Composite.UnitTests.MaskingPass",
+	EAutomationTestFlags_ApplicationContextMask | EAutomationTestFlags::EngineFilter | EAutomationTestFlags::NonNullRHI)
+
+bool FCompositeMaskingPassTest::RunTest(const FString& Parameter)
+{
+	using namespace UE::Composite::Private;
+	using namespace UE::CompositeCore;
+	using namespace Masking;
+
+	const int32 NumScenarios = UE_ARRAY_COUNT(GScenarios);
+
+	TArray<FLinearColor> ResultRDG;
+	ResultRDG.SetNumZeroed(NumScenarios);
+
+	TArray<FMaskingPassProxy> Proxies;
+	Proxies.Reserve(NumScenarios);
+	for (int32 i = 0; i < NumScenarios; ++i)
+	{
+		const FMaskingScenario& S = GScenarios[i];
+		FPassInputDeclArray InputDecls;
+		InputDecls.AddDefaulted_GetRef().Set<FPassInternalResourceDesc>(FPassInternalResourceDesc{});
+		InputDecls.AddDefaulted_GetRef().Set<FPassInternalResourceDesc>(FPassInternalResourceDesc{});
+		FMaskingPassProxy Proxy(MoveTemp(InputDecls));
+		Proxy.MaskSourceChannel = S.Channel;
+		Proxy.bUnpremult        = S.bUnpremult;
+		Proxy.bInvert           = S.bInvert;
+		Proxies.Emplace(MoveTemp(Proxy));
+	}
+
+	FCompositeTestView TestView;
+	DispatchAndWait(TestView.Get(),
+		[&ResultRDG, NumScenarios, Proxies = MoveTemp(Proxies)](FRDGBuilder& GraphBuilder, const FSceneView& View) mutable
+		{
+			const FIntRect PixelRect(0, 0, GTestViewSize.X, GTestViewSize.Y);
+
+			for (int32 i = 0; i < NumScenarios; ++i)
+			{
+				const FMaskingScenario& S = GScenarios[i];
+
+				const FRDGTextureDesc MainDesc = FRDGTextureDesc::Create2D(
+					GTestViewSize, PF_A32B32G32R32F,
+					FClearValueBinding::Black, ETextureCreateFlags::None);
+				FRDGTextureRef InputTexture = GraphBuilder.CreateTexture(MainDesc, TEXT("MaskingTest.Input"));
+
+				const FRDGTextureDesc MaskDesc = FRDGTextureDesc::Create2D(
+					GTestViewSize, PF_A32B32G32R32F,
+					FClearValueBinding::Black, ETextureCreateFlags::None);
+				FRDGTextureRef MaskTexture = GraphBuilder.CreateTexture(MaskDesc, TEXT("MaskingTest.Mask"));
+
+				const FRDGTextureDesc OutputDesc = FRDGTextureDesc::Create2D(
+					GTestViewSize, PF_A32B32G32R32F,
+					FClearValueBinding::Black, ETextureCreateFlags::RenderTargetable);
+				FRDGTextureRef OutputTexture = GraphBuilder.CreateTexture(OutputDesc, TEXT("MaskingTest.Output"));
+
+				{
+					FMaskingTestMainUploadParameters* Params = GraphBuilder.AllocParameters<FMaskingTestMainUploadParameters>();
+					Params->InputTexture = InputTexture;
+
+					const FLinearColor PixelValue = S.MainInput;
+					GraphBuilder.AddPass(
+						RDG_EVENT_NAME("MaskingTest.Upload[%d]", i), Params,
+						ERDGPassFlags::Readback,
+						[PixelValue, Params](FRHICommandListImmediate& RHICmdList)
+						{
+							uint32 Stride = 0;
+							void* Data = RHICmdList.LockTexture2D(Params->InputTexture->GetRHI(), 0, RLM_WriteOnly, Stride, false);
+							for (int32 Row = 0; Row < GTestViewSize.Y; ++Row)
+							{
+								FLinearColor* RowPtr = reinterpret_cast<FLinearColor*>(static_cast<uint8*>(Data) + Row * Stride);
+								for (int32 Col = 0; Col < GTestViewSize.X; ++Col)
+								{
+									RowPtr[Col] = PixelValue;
+								}
+							}
+							RHICmdList.UnlockTexture2D(Params->InputTexture->GetRHI(), 0, false);
+						});
+				}
+
+				{
+					FMaskingTestMaskUploadParameters* Params = GraphBuilder.AllocParameters<FMaskingTestMaskUploadParameters>();
+					Params->MaskTexture = MaskTexture;
+
+					const FLinearColor MaskValue = S.MaskInput;
+					GraphBuilder.AddPass(
+						RDG_EVENT_NAME("MaskingTest.MaskUpload[%d]", i), Params,
+						ERDGPassFlags::Readback,
+						[MaskValue, Params](FRHICommandListImmediate& RHICmdList)
+						{
+							uint32 Stride = 0;
+							void* Data = RHICmdList.LockTexture2D(Params->MaskTexture->GetRHI(), 0, RLM_WriteOnly, Stride, false);
+							for (int32 Row = 0; Row < GTestViewSize.Y; ++Row)
+							{
+								FLinearColor* RowPtr = reinterpret_cast<FLinearColor*>(static_cast<uint8*>(Data) + Row * Stride);
+								for (int32 Col = 0; Col < GTestViewSize.X; ++Col)
+								{
+									RowPtr[Col] = MaskValue;
+								}
+							}
+							RHICmdList.UnlockTexture2D(Params->MaskTexture->GetRHI(), 0, false);
+						});
+				}
+
+				FPassInputArray Inputs;
+				{
+					FPassInput& InMain = Inputs.GetArray().AddDefaulted_GetRef();
+					InMain.Texture         = FScreenPassTexture{ InputTexture, PixelRect };
+					InMain.Metadata.Filter = SF_Point;
+				}
+				{
+					FPassInput& InMask = Inputs.GetArray().AddDefaulted_GetRef();
+					InMask.Texture         = FScreenPassTexture{ MaskTexture, PixelRect };
+					InMask.Metadata.Filter = SF_Point;
+				}
+				Inputs.OverrideOutput = FScreenPassRenderTarget{
+					FScreenPassTexture{ OutputTexture, PixelRect }, ERenderTargetLoadAction::ENoAction };
+
+				FPassContext PassContext;
+				PassContext.OutputViewRect = PixelRect;
+
+				Proxies[i].Add(GraphBuilder, View, Inputs, PassContext);
+
+				{
+					FMaskingTestDownloadParameters* Params = GraphBuilder.AllocParameters<FMaskingTestDownloadParameters>();
+					Params->OutputTexture = OutputTexture;
+
+					FLinearColor* ResultPtr = &ResultRDG[i];
+					GraphBuilder.AddPass(
+						RDG_EVENT_NAME("MaskingTest.Download[%d]", i), Params,
+						ERDGPassFlags::Readback,
+						[ResultPtr, Params](FRHICommandListImmediate& RHICmdList)
+						{
+							uint32 Stride = 0;
+							const void* Data = RHICmdList.LockTexture2D(Params->OutputTexture->GetRHI(), 0, RLM_ReadOnly, Stride, false);
+							const FLinearColor* RowPtr = reinterpret_cast<const FLinearColor*>(
+								static_cast<const uint8*>(Data) + (GTestViewSize.Y / 2) * Stride);
+							FMemory::Memcpy(ResultPtr, &RowPtr[GTestViewSize.X / 2], sizeof(FLinearColor));
+							RHICmdList.UnlockTexture2D(Params->OutputTexture->GetRHI(), 0, false);
+						});
+				}
+			}
+		});
+
+	constexpr float Tolerance = 1e-4f;
+	for (int32 i = 0; i < NumScenarios; ++i)
+	{
+		const FMaskingScenario& S = GScenarios[i];
+		const FLinearColor Expected = MaskingCPU(S.MainInput, S.MaskInput, S.Channel, S.bUnpremult, S.bInvert);
+		const FLinearColor& Actual  = ResultRDG[i];
+
+		UTEST_EQUAL_TOLERANCE(S.Label, Actual.R, Expected.R, Tolerance);
+		UTEST_EQUAL_TOLERANCE(S.Label, Actual.G, Expected.G, Tolerance);
+		UTEST_EQUAL_TOLERANCE(S.Label, Actual.B, Expected.B, Tolerance);
+		UTEST_EQUAL_TOLERANCE(S.Label, Actual.A, Expected.A, Tolerance);
+	}
+
+	return true;
+}
+
+#endif // WITH_DEV_AUTOMATION_TESTS
+
+} // namespace UE::Composite::Tests
diff --git a/Engine/Plugins/Compositing/Composite/Source/Composite/Private/Tests/CompositePassTranslucencyTest.cpp b/Engine/Plugins/Compositing/Composite/Source/Composite/Private/Tests/CompositePassTranslucencyTest.cpp
new file mode 100644
--- /dev/null
+++ b/Engine/Plugins/Compositing/Composite/Source/Composite/Private/Tests/CompositePassTranslucencyTest.cpp
@@ -0,0 +1,207 @@
+// Copyright Epic Games, Inc. All Rights Reserved.
+
+#include "CoreTypes.h"
+#include "RenderGraphBuilder.h"
+#include "RHI.h"
+#include "ScreenPass.h"
+
+#include "Passes/CompositePassTranslucency.h"
+#include "Tests/CompositeTestHelpers.h"
+
+#if WITH_DEV_AUTOMATION_TESTS
+#include "Misc/AutomationTest.h"
+#endif
+
+namespace UE::Composite::Tests
+{
+
+#if WITH_DEV_AUTOMATION_TESTS
+
+// -- Upload/download parameter structs --
+
+BEGIN_SHADER_PARAMETER_STRUCT(FTranslucencyTestUploadParameters, )
+	RDG_TEXTURE_ACCESS(InputTexture, ERHIAccess::CopyDest)
+END_SHADER_PARAMETER_STRUCT()
+
+BEGIN_SHADER_PARAMETER_STRUCT(FTranslucencyTestDownloadParameters, )
+	RDG_TEXTURE_ACCESS(OutputTexture, ERHIAccess::CopySrc)
+END_SHADER_PARAMETER_STRUCT()
+
+namespace Translucency  // unique name prevents symbol collision in unity builds
+{
+
+// -- CPU reference --
+
+static FLinearColor TranslucencyCPU(
+	FLinearColor In,
+	ECompositeAlphaPremultiplication PremultOp,
+	float Fade,
+	ECompositeAlphaOverride Override)
+{
+	if (PremultOp == ECompositeAlphaPremultiplication::Premultiply)
+	{
+		In.R *= In.A;
+		In.G *= In.A;
+		In.B *= In.A;
+		In.R *= Fade; In.G *= Fade; In.B *= Fade; In.A *= Fade;
+	}
+	else if (PremultOp == ECompositeAlphaPremultiplication::Unpremultiply)
+	{
+		const float Rcp = In.A > 0.f ? 1.f / In.A : 0.f;
+		In.R *= Rcp;
+		In.G *= Rcp;
+		In.B *= Rcp;
+		In.A *= Fade;
+	}
+	else
+	{
+		In.R *= Fade; In.G *= Fade; In.B *= Fade; In.A *= Fade;
+	}
+
+	if (Override == ECompositeAlphaOverride::Zero)  { In.A = 0.f; }
+	if (Override == ECompositeAlphaOverride::One)   { In.A = 1.f; }
+
+	return In;
+}
+
+// -- Test scenarios --
+
+struct FTranslucencyScenario
+{
+	const TCHAR*                       Label;
+	FLinearColor                       Input;
+	ECompositeAlphaPremultiplication   PremultOp;
+	float                              Fade;
+	ECompositeAlphaOverride            Override;
+};
+
+static const FTranslucencyScenario GScenarios[] =
+{
+	{ TEXT("None+Fade"),         { 0.8f, 0.6f, 0.4f, 1.0f }, ECompositeAlphaPremultiplication::None,          0.5f, ECompositeAlphaOverride::None },
+	{ TEXT("Premultiply"),       { 0.8f, 0.6f, 0.4f, 0.5f }, ECompositeAlphaPremultiplication::Premultiply,   1.0f, ECompositeAlphaOverride::None },
+	{ TEXT("Unpremultiply"),     { 0.4f, 0.3f, 0.2f, 0.5f }, ECompositeAlphaPremultiplication::Unpremultiply, 0.8f, ECompositeAlphaOverride::None },
+	{ TEXT("OverrideAlphaZero"), { 0.8f, 0.6f, 0.4f, 0.7f }, ECompositeAlphaPremultiplication::None,          1.0f, ECompositeAlphaOverride::Zero },
+	{ TEXT("OverrideAlphaOne"),  { 0.8f, 0.6f, 0.4f, 0.0f }, ECompositeAlphaPremultiplication::None,          1.0f, ECompositeAlphaOverride::One  },
+	{ TEXT("UnpremultiplyAlpha0"),{ 0.4f, 0.3f, 0.2f, 0.0f }, ECompositeAlphaPremultiplication::Unpremultiply,1.0f, ECompositeAlphaOverride::None },
+};
+
+}  // namespace Translucency
+
+// -- Test --
+
+IMPLEMENT_SIMPLE_AUTOMATION_TEST(FCompositeTranslucencyPassTest, "Composite.UnitTests.TranslucencyPass",
+	EAutomationTestFlags_ApplicationContextMask | EAutomationTestFlags::EngineFilter | EAutomationTestFlags::NonNullRHI)
+
+bool FCompositeTranslucencyPassTest::RunTest(const FString& Parameter)
+{
+	using namespace UE::Composite::Private;
+	using namespace UE::CompositeCore;
+	using namespace Translucency;
+
+	const int32 NumScenarios = UE_ARRAY_COUNT(GScenarios);
+
+	TArray<FLinearColor> ResultRDG;
+	ResultRDG.SetNumZeroed(NumScenarios);
+
+	TArray<FTranslucencyPassProxy> Proxies;
+	Proxies.Reserve(NumScenarios);
+	for (int32 i = 0; i < NumScenarios; ++i)
+	{
+		const FTranslucencyScenario& S = GScenarios[i];
+		FTranslucencyPassProxy Proxy(GetDefaultInputDeclArray());
+		Proxy.PremultOp           = S.PremultOp;
+		Proxy.Fade                = S.Fade;
+		Proxy.OverrideOutputAlpha = S.Override;
+		Proxies.Emplace(MoveTemp(Proxy));
+	}
+
+	FCompositeTestView TestView;
+	DispatchAndWait(TestView.Get(),
+		[&ResultRDG, NumScenarios, Proxies = MoveTemp(Proxies)](FRDGBuilder& GraphBuilder, const FSceneView& View) mutable
+		{
+			// 1×1 textures are sufficient: the translucency pass operates per-pixel uniformly.
+			const FIntRect PixelRect(0, 0, 1, 1);
+
+			for (int32 i = 0; i < NumScenarios; ++i)
+			{
+				const FTranslucencyScenario& S = GScenarios[i];
+
+				const FRDGTextureDesc InputDesc = FRDGTextureDesc::Create2D(
+					FIntPoint(1, 1), PF_A32B32G32R32F,
+					FClearValueBinding::Black, ETextureCreateFlags::None);
+				FRDGTextureRef InputTexture = GraphBuilder.CreateTexture(InputDesc, TEXT("TranslucencyTest.Input"));
+
+				const FRDGTextureDesc OutputDesc = FRDGTextureDesc::Create2D(
+					FIntPoint(1, 1), PF_A32B32G32R32F,
+					FClearValueBinding::Black, ETextureCreateFlags::RenderTargetable);
+				FRDGTextureRef OutputTexture = GraphBuilder.CreateTexture(OutputDesc, TEXT("TranslucencyTest.Output"));
+
+				{
+					FTranslucencyTestUploadParameters* Params = GraphBuilder.AllocParameters<FTranslucencyTestUploadParameters>();
+					Params->InputTexture = InputTexture;
+
+					const FLinearColor PixelValue = S.Input;
+					GraphBuilder.AddPass(
+						RDG_EVENT_NAME("TranslucencyTest.Upload[%d]", i), Params,
+						ERDGPassFlags::Readback,
+						[PixelValue, Params](FRHICommandListImmediate& RHICmdList)
+						{
+							uint32 Stride = 0;
+							void* Data = RHICmdList.LockTexture2D(Params->InputTexture->GetRHI(), 0, RLM_WriteOnly, Stride, false);
+							FMemory::Memcpy(Data, &PixelValue, sizeof(FLinearColor));
+							RHICmdList.UnlockTexture2D(Params->InputTexture->GetRHI(), 0, false);
+						});
+				}
+
+				FPassInputArray Inputs;
+				{
+					FPassInput& In = Inputs.GetArray().AddDefaulted_GetRef();
+					In.Texture         = FScreenPassTexture{ InputTexture, PixelRect };
+					In.Metadata.Filter = SF_Point;
+				}
+				Inputs.OverrideOutput = FScreenPassRenderTarget{
+					FScreenPassTexture{ OutputTexture, PixelRect }, ERenderTargetLoadAction::ENoAction };
+
+				FPassContext PassContext;
+				PassContext.OutputViewRect = PixelRect;
+
+				Proxies[i].Add(GraphBuilder, View, Inputs, PassContext);
+
+				{
+					FTranslucencyTestDownloadParameters* Params = GraphBuilder.AllocParameters<FTranslucencyTestDownloadParameters>();
+					Params->OutputTexture = OutputTexture;
+
+					FLinearColor* ResultPtr = &ResultRDG[i];
+					GraphBuilder.AddPass(
+						RDG_EVENT_NAME("TranslucencyTest.Download[%d]", i), Params,
+						ERDGPassFlags::Readback,
+						[ResultPtr, Params](FRHICommandListImmediate& RHICmdList)
+						{
+							uint32 Stride = 0;
+							const void* Data = RHICmdList.LockTexture2D(Params->OutputTexture->GetRHI(), 0, RLM_ReadOnly, Stride, false);
+							FMemory::Memcpy(ResultPtr, Data, sizeof(FLinearColor));
+							RHICmdList.UnlockTexture2D(Params->OutputTexture->GetRHI(), 0, false);
+						});
+				}
+			}
+		});
+
+	constexpr float Tolerance = 1e-4f;
+	for (int32 i = 0; i < NumScenarios; ++i)
+	{
+		const FTranslucencyScenario& S = GScenarios[i];
+		const FLinearColor Expected = TranslucencyCPU(S.Input, S.PremultOp, S.Fade, S.Override);
+		const FLinearColor& Actual  = ResultRDG[i];
+
+		UTEST_EQUAL_TOLERANCE(S.Label, Actual.R, Expected.R, Tolerance);
+		UTEST_EQUAL_TOLERANCE(S.Label, Actual.G, Expected.G, Tolerance);
+		UTEST_EQUAL_TOLERANCE(S.Label, Actual.B, Expected.B, Tolerance);
+		UTEST_EQUAL_TOLERANCE(S.Label, Actual.A, Expected.A, Tolerance);
+	}
+
+	return true;
+}
+
+#endif // WITH_DEV_AUTOMATION_TESTS
+
+} // namespace UE::Composite::Tests
diff --git a/Engine/Plugins/Compositing/Composite/Source/Composite/Private/Tests/CompositeTestHelpers.h b/Engine/Plugins/Compositing/Composite/Source/Composite/Private/Tests/CompositeTestHelpers.h
new file mode 100644
--- /dev/null
+++ b/Engine/Plugins/Compositing/Composite/Source/Composite/Private/Tests/CompositeTestHelpers.h
@@ -0,0 +1,78 @@
+// Copyright Epic Games, Inc. All Rights Reserved.
+
+#pragma once
+
+#include "CoreTypes.h"
+#include "HAL/Event.h"
+#include "HAL/PlatformProcess.h"
+#include "RenderGraphBuilder.h"
+#include "RenderingThread.h"
+#include "RHI.h"
+#include "SceneView.h"
+#include "UnrealClient.h"
+#include "Templates/Function.h"
+
+namespace UE::Composite::Tests
+{
+
+// Shared viewport size used by all unit tests (wide enough to keep center pixels
+// clear of blur boundaries; small enough for fast texture allocation and readback).
+inline const FIntPoint GTestViewSize(64, 32);
+
+// Owns all game-thread resources needed to create a minimal FSceneView for
+// render-thread dispatch: a dummy render target, a view family, and the view itself.
+// Members are declared in construction order so that DummyRT is alive when ViewFamily
+// is constructed (ViewFamily stores a pointer to the render target).
+struct FCompositeTestView
+{
+private:
+	struct FDummyRenderTarget final : public FRenderTarget
+	{
+		virtual FIntPoint GetSizeXY() const override { return GTestViewSize; }
+	};
+
+	FDummyRenderTarget     DummyRT;
+	FSceneViewFamily       ViewFamily;
+	TUniquePtr<FSceneView> TestView;
+
+public:
+	FCompositeTestView()
+		: ViewFamily(FSceneViewFamily::ConstructionValues(&DummyRT, nullptr, FEngineShowFlags(ESFIM_Game))
+			.SetTime(FGameTime()))
+	{
+		FSceneViewInitOptions ViewInitOptions;
+		ViewInitOptions.ViewFamily         = &ViewFamily;
+		ViewInitOptions.ViewOrigin         = FVector::ZeroVector;
+		ViewInitOptions.ViewRotationMatrix = FMatrix::Identity;
+		ViewInitOptions.ProjectionMatrix   = FMatrix::Identity;
+		ViewInitOptions.SetViewRectangle(FIntRect(FIntPoint::ZeroValue, GTestViewSize));
+		TestView = MakeUnique<FSceneView>(ViewInitOptions);
+	}
+
+	FSceneView* Get() const { return TestView.Get(); }
+};
+
+// Enqueues a batch of RDG passes on the render thread and blocks the game thread until
+// they complete. Work receives an open FRDGBuilder and the view; GraphBuilder.Execute()
+// and BlockUntilGPUIdle() are called internally so callers only need to add passes.
+inline void DispatchAndWait(
+	FSceneView* ViewPtr,
+	TUniqueFunction<void(FRDGBuilder&, const FSceneView&)> Work)
+{
+	FEvent* Signal = FGenericPlatformProcess::GetSynchEventFromPool(true);
+
+	ENQUEUE_RENDER_COMMAND(CompositeTestDispatch)(
+		[ViewPtr, Signal, Work = MoveTemp(Work)](FRHICommandListImmediate& RHICmdList) mutable
+		{
+			FRDGBuilder GraphBuilder(RHICmdList);
+			Work(GraphBuilder, *ViewPtr);
+			GraphBuilder.Execute();
+			RHICmdList.BlockUntilGPUIdle();
+			Signal->Trigger();
+		});
+
+	Signal->Wait();
+	FGenericPlatformProcess::ReturnSynchEventToPool(Signal);
+}
+
+} // namespace UE::Composite::Tests
diff --git a/Engine/Plugins/Compositing/Composite/Source/Composite/Public/Passes/CompositePassBlur.h b/Engine/Plugins/Compositing/Composite/Source/Composite/Public/Passes/CompositePassBlur.h
--- a/Engine/Plugins/Compositing/Composite/Source/Composite/Public/Passes/CompositePassBlur.h
+++ b/Engine/Plugins/Compositing/Composite/Source/Composite/Public/Passes/CompositePassBlur.h
@@ -46,4 +46,32 @@ public:
 	bool bAlphaOnly = false;
 };
 
+namespace UE
+{
+	namespace Composite
+	{
+		namespace Private
+		{
+			using namespace CompositeCore;
+
+			/** Render-thread proxy for the Gaussian blur pass. Runs a separable horizontal + vertical convolution over the input. */
+			class FCompositePassBlurProxy : public FCompositeCorePassProxy
+			{
+			public:
+				IMPLEMENT_COMPOSITE_PASS(FCompositePassBlurProxy);
+
+				using FCompositeCorePassProxy::FCompositeCorePassProxy;
+
+				/** Adds the Gaussian blur RDG passes (horizontal then vertical) and returns the blurred output. */
+				FPassTexture Add(FRDGBuilder& GraphBuilder, const FSceneView& InView, const FPassInputArray& Inputs, const FPassContext& PassContext) const override;
+
+				/** Blur kernel half-width in pixels; larger values produce a softer result. */
+				int32 Radius    = 1;
+				/** When true, only the alpha channel is blurred; RGB is passed through unchanged. */
+				bool bAlphaOnly = false;
+			};
+		}
+	}
+}
+
 #undef UE_API
diff --git a/Engine/Plugins/Compositing/Composite/Source/Composite/Public/Passes/CompositePassCenteredScale.h b/Engine/Plugins/Compositing/Composite/Source/Composite/Public/Passes/CompositePassCenteredScale.h
--- a/Engine/Plugins/Compositing/Composite/Source/Composite/Public/Passes/CompositePassCenteredScale.h
+++ b/Engine/Plugins/Compositing/Composite/Source/Composite/Public/Passes/CompositePassCenteredScale.h
@@ -83,4 +83,30 @@ private:
 	mutable TOptional<FVector2f> CachedScaleUV;
 };
 
+namespace UE
+{
+	namespace Composite
+	{
+		namespace Private
+		{
+			using namespace CompositeCore;
+
+			/** Render-thread proxy for the centered-scale pass. Remaps input UVs around the viewport center: UV_in = (UV_out − 0.5) * ScaleUV + 0.5. */
+			class FCenteredScalePassProxy : public FCompositeCorePassProxy
+			{
+			public:
+				IMPLEMENT_COMPOSITE_PASS(FCenteredScalePassProxy);
+
+				using FCompositeCorePassProxy::FCompositeCorePassProxy;
+
+				/** Adds the centered-scale RDG pass; bilinear-resamples the input with UVs remapped around the viewport center and returns the scaled output. */
+				FPassTexture Add(FRDGBuilder& GraphBuilder, const FSceneView& InView, const FPassInputArray& Inputs, const FPassContext& PassContext) const override;
+
+				/** Per-axis UV scale factor. Values > 1 zoom in; values < 1 zoom out. (1, 1) is identity. */
+				FVector2f ScaleUV = FVector2f::One();
+			};
+		}
+	}
+}
+
 #undef UE_API
diff --git a/Engine/Plugins/Compositing/Composite/Source/Composite/Public/Passes/CompositePassColorGrading.h b/Engine/Plugins/Compositing/Composite/Source/Composite/Public/Passes/CompositePassColorGrading.h
--- a/Engine/Plugins/Compositing/Composite/Source/Composite/Public/Passes/CompositePassColorGrading.h
+++ b/Engine/Plugins/Compositing/Composite/Source/Composite/Public/Passes/CompositePassColorGrading.h
@@ -93,7 +93,7 @@ namespace UE
 	{
 		namespace Private
 		{
-			/** Render pass proxy implementation. */
+			/** Render-thread proxy for the color grading pass. Applies white balance and shadow/midtone/highlight color grading to the input. */
 			class FColorGradingCompositePassProxy : public FCompositeCorePassProxy
 			{
 			public:
@@ -102,6 +102,7 @@ namespace UE
 				/** Constructor */
 				FColorGradingCompositePassProxy(CompositeCore::FPassInputDeclArray InPassDeclaredInputs);
 
+				/** Adds the color grading RDG pass; applies temperature and color grading adjustments and returns the graded output. */
 				CompositeCore::FPassTexture Add(FRDGBuilder& GraphBuilder, const FSceneView& InView, const CompositeCore::FPassInputArray& Inputs, const CompositeCore::FPassContext& PassContext) const override;
 
 				/** Color is divided by its alpha channel before processing, and multiplied again after. */
diff --git a/Engine/Plugins/Compositing/Composite/Source/Composite/Public/Passes/CompositePassColorKeyer.h b/Engine/Plugins/Compositing/Composite/Source/Composite/Public/Passes/CompositePassColorKeyer.h
--- a/Engine/Plugins/Compositing/Composite/Source/Composite/Public/Passes/CompositePassColorKeyer.h
+++ b/Engine/Plugins/Compositing/Composite/Source/Composite/Public/Passes/CompositePassColorKeyer.h
@@ -147,5 +147,47 @@ private:
 	friend class FCompositePassColorKeyerCustomization;
 };
 
+namespace UE
+{
+	namespace Composite
+	{
+		namespace Private
+		{
+			using namespace CompositeCore;
+			/** Render-thread proxy for the color keyer pass. Performs chroma-key extraction and optional despill on the input. */
+			class FCompositePassColorKeyerProxy : public FCompositeCorePassProxy
+			{
+			public:
+				IMPLEMENT_COMPOSITE_PASS(FCompositePassColorKeyerProxy);
+				using FCompositeCorePassProxy::FCompositeCorePassProxy;
+
+				/** Adds the color keyer RDG pass and returns a premultiplied keyed RGBA output. */
+				FPassTexture Add(FRDGBuilder& GraphBuilder, const FSceneView& InView, const FPassInputArray& Inputs, const FPassContext& PassContext) const override;
+
+				/** Background screen color used to drive the key algorithm. */
+				ECompositeColorKeyerScreenType    ScreenType             = ECompositeColorKeyerScreenType::Green;
+				/** Static background reference color; used when no clean plate is provided. */
+				FLinearColor                      KeyColor;
+				/** Shader parameter vector: (RedWeight, GreenWeight, BlueWeight, DespillStrength). */
+				FVector4f                         Params0                = FVector4f(0.0f, 0.0f, 0.0f, 0.0f);
+				/** Shader parameter vector: (AlphaMin, AlphaMax, DevignetteStrength, unused). */
+				FVector4f                         Params1                = FVector4f(0.0f, 0.0f, 0.0f, 0.0f);
+				/** Which output channel to display: final composite key, fill layer, or raw alpha matte. */
+				ECompositeColorKeyerVisualization Visualization          = ECompositeColorKeyerVisualization::Key;
+				/** When true, devignette correction is reversed before outputting the keyed result. */
+				bool                              bPreserveVignetteAfterKey = true;
+				/** Denoising pre-pass applied to the foreground before keying. */
+				ECompositeDenoiseMethod           DenoiseMethod          = ECompositeDenoiseMethod::None;
+				/** When true, only despill is applied; alpha is left unchanged. */
+				bool                              bDespillOnly           = false;
+				/** When true, the output alpha is inverted (1 − alpha). */
+				bool                              bInvertAlpha           = false;
+				/** Weak pointer to the clean plate texture used as the per-pixel background reference. Null when KeyColor is used instead. */
+				TWeakObjectPtr<UTexture>          CleanPlateWeakPtr;
+			};
+		}
+	}
+}
+
 #undef UE_API
 
diff --git a/Engine/Plugins/Compositing/Composite/Source/Composite/Public/Passes/CompositePassMasking.h b/Engine/Plugins/Compositing/Composite/Source/Composite/Public/Passes/CompositePassMasking.h
--- a/Engine/Plugins/Compositing/Composite/Source/Composite/Public/Passes/CompositePassMasking.h
+++ b/Engine/Plugins/Compositing/Composite/Source/Composite/Public/Passes/CompositePassMasking.h
@@ -53,5 +53,35 @@ public:
 	bool bInvert;
 };
 
+namespace UE
+{
+	namespace Composite
+	{
+		namespace Private
+		{
+			using namespace CompositeCore;
+
+			/** Render-thread proxy for the alpha masking pass. Takes two inputs: [0] main, [1] mask. */
+			class FMaskingPassProxy : public FCompositeCorePassProxy
+			{
+			public:
+				IMPLEMENT_COMPOSITE_PASS(FMaskingPassProxy);
+
+				using FCompositeCorePassProxy::FCompositeCorePassProxy;
+
+				/** Adds the masking RDG pass; multiplies the main input's alpha by the chosen mask channel and returns the masked output. */
+				FPassTexture Add(FRDGBuilder& GraphBuilder, const FSceneView& InView, const FPassInputArray& Inputs, const FPassContext& PassContext) const override;
+
+				/** Index of the mask texture channel to read (0=R, 1=G, 2=B, 3=A). */
+				uint32 MaskSourceChannel = 0;
+				/** When true, RGB is divided by alpha before masking and multiplied back after, preserving straight-alpha color values. */
+				bool   bUnpremult = true;
+				/** When true, the mask value is inverted (1 − mask) before it is applied. */
+				bool   bInvert    = false;
+			};
+		}
+	}
+}
+
 #undef UE_API
 
diff --git a/Engine/Plugins/Compositing/Composite/Source/Composite/Public/Passes/CompositePassTranslucency.h b/Engine/Plugins/Compositing/Composite/Source/Composite/Public/Passes/CompositePassTranslucency.h
--- a/Engine/Plugins/Compositing/Composite/Source/Composite/Public/Passes/CompositePassTranslucency.h
+++ b/Engine/Plugins/Compositing/Composite/Source/Composite/Public/Passes/CompositePassTranslucency.h
@@ -62,6 +62,7 @@ namespace UE
 		{
 			using namespace CompositeCore;
 
+			/** Render-thread proxy for the translucency pass. Applies alpha premultiplication, fade, and optional alpha override. */
 			class FTranslucencyPassProxy : public FCompositeCorePassProxy
 			{
 			public:
@@ -69,10 +70,14 @@ namespace UE
 
 				using FCompositeCorePassProxy::FCompositeCorePassProxy;
 
+				/** Adds the translucency RDG pass and returns the modified output texture. */
 				FPassTexture Add(FRDGBuilder& GraphBuilder, const FSceneView& InView, const FPassInputArray& Inputs, const FPassContext& PassContext) const override;
 
+				/** Whether to premultiply, unpremultiply, or leave alpha unchanged before applying Fade. */
 				ECompositeAlphaPremultiplication PremultOp = ECompositeAlphaPremultiplication::None;
+				/** Uniform multiplier applied to the output (0 = fully transparent, 1 = unchanged). */
 				float Fade = 1.0f;
+				/** Optionally force the output alpha to 0 or 1 after all other operations. */
 				ECompositeAlphaOverride OverrideOutputAlpha = ECompositeAlphaOverride::None;
 			};
 		}
