Composure: Expose erosion/dilation pass support, which was previously only used implicitely by the built-in custom render pass.

* Added erosion mode as a new shader permutation.
* Larger dilations/erosions are support on our new pass by running the same shader iteratively.

#Synced-CL 52419094

diff --git a/Engine/Plugins/Compositing/Composite/Source/Composite/Private/Passes/CompositePassDilation.cpp b/Engine/Plugins/Compositing/Composite/Source/Composite/Private/Passes/CompositePassDilation.cpp
new file mode 100644
--- /dev/null
+++ b/Engine/Plugins/Compositing/Composite/Source/Composite/Private/Passes/CompositePassDilation.cpp
@@ -0,0 +1,181 @@
+// Copyright Epic Games, Inc. All Rights Reserved.
+
+#include "Passes/CompositePassDilation.h"
+
+#include "Passes/CompositeCorePassProxy.h"
+#include "RenderGraphBuilder.h"
+#include "RenderGraphUtils.h"
+#include "SceneView.h"
+#include "ScreenPass.h"
+
+DECLARE_GPU_STAT_NAMED(FCompositeDilation, TEXT("Composite.Dilation"));
+
+class FCompositePassDilationShader : public FGlobalShader
+{
+	DECLARE_GLOBAL_SHADER(FCompositePassDilationShader);
+	SHADER_USE_PARAMETER_STRUCT(FCompositePassDilationShader, FGlobalShader);
+
+	// DILATION_SIZE drives the LDS tile halo and the neighbor kernel:
+	//   1 → 8-neighbor (3×3 minus center), 1 pixel per dispatch
+	//   2 → 24-neighbor (5×5 minus center), 2 pixels per dispatch
+	// The C++ dispatch loop uses DILATION_SIZE=2 greedily, then DILATION_SIZE=1 for any
+	// remaining pixel — same reach as N×DILATION_SIZE=1 but up to half the dispatch count.
+	class FDilationSize : SHADER_PERMUTATION_RANGE_INT("DILATION_SIZE", 1, 2);
+	class FErode : SHADER_PERMUTATION_BOOL("MORPHOLOGY_ERODE");
+	using FPermutationDomain = TShaderPermutationDomain<FDilationSize, FErode>;
+
+	BEGIN_SHADER_PARAMETER_STRUCT(FParameters, )
+		SHADER_PARAMETER_STRUCT(FScreenPassTextureInput, Input)
+		SHADER_PARAMETER_STRUCT(FScreenPassTextureViewportParameters, Output)
+		SHADER_PARAMETER_RDG_TEXTURE_UAV(RWTexture2D<float4>, RWOutputTexture)
+		SHADER_PARAMETER(uint32, bInvertedAlpha)
+		SHADER_PARAMETER(uint32, bOpacifyOutput)
+		SHADER_PARAMETER(float, AlphaThreshold)
+	END_SHADER_PARAMETER_STRUCT()
+
+	static inline void ModifyCompilationEnvironment(const FGlobalShaderPermutationParameters& Parameters, FShaderCompilerEnvironment& OutEnvironment)
+	{
+		OutEnvironment.SetDefine(TEXT("TILE_SIZE"), ThreadGroupSize);
+
+		FGlobalShader::ModifyCompilationEnvironment(Parameters, OutEnvironment);
+	}
+
+	static const uint32 ThreadGroupSize = 16;
+};
+
+IMPLEMENT_GLOBAL_SHADER(FCompositePassDilationShader, "/Plugin/CompositeCore/Private/CompositeCoreDilate.usf", "MainCS", SF_Compute);
+
+namespace UE::Composite::Private
+{
+
+FPassTexture FCompositePassDilationProxy::Add(FRDGBuilder& GraphBuilder, const FSceneView& InView, const FPassInputArray& Inputs, const FPassContext& PassContext) const
+{
+	RDG_EVENT_SCOPE_STAT(GraphBuilder, FCompositeDilation, "Composite.Dilation");
+
+	check(ValidateInputs(Inputs));
+	check(Size >= 1);
+
+	const UE::CompositeCore::FResourceMetadata& Metadata = Inputs[0].Metadata;
+	const FScreenPassTexture& InputTexture = Inputs[0].Texture;
+
+	FGlobalShaderMap* GlobalShaderMap = GetGlobalShaderMap(InView.GetFeatureLevel());
+
+	FScreenPassRenderTarget Output = Inputs.OverrideOutput;
+	if (Output.IsValid())
+	{
+		if (!ensureMsgf(!!(Output.Texture->Desc.Flags & TexCreate_UAV), TEXT("CompositePassDilation: OverrideOutput texture '%s' was not created with TexCreate_UAV; ignoring override and allocating an internal output instead."), Output.Texture->Name))
+		{
+			Output = {};
+		}
+	}
+	if (!Output.IsValid())
+	{
+		Output = CreateOutputRenderTarget(GraphBuilder, InView, PassContext.OutputViewRect, InputTexture.Texture->Desc, TEXT("CompositeDilationOutput"));
+	}
+
+	// Input and Output viewports must be the same size: the dilation kernel operates at a fixed
+	// resolution throughout the ping-pong chain, and the final iteration writes directly to Output.
+	const FIntPoint TextureSize = InputTexture.ViewRect.Size();
+	checkf(Output.ViewRect.Size() == TextureSize,
+		TEXT("CompositePassDilation: Input viewport (%dx%d) and Output viewport (%dx%d) must be the same size."),
+		TextureSize.X, TextureSize.Y, Output.ViewRect.Size().X, Output.ViewRect.Size().Y);
+
+	const FRDGTextureDesc PingPongDesc = FRDGTextureDesc::Create2D(TextureSize, Output.Texture->Desc.Format, FClearValueBinding::Black, TexCreate_UAV | TexCreate_ShaderResource | TexCreate_RenderTargetable);
+	FRDGTextureRef PingPong[2] = {
+		GraphBuilder.CreateTexture(PingPongDesc, TEXT("CompositeDilationPingPong0")),
+		GraphBuilder.CreateTexture(PingPongDesc, TEXT("CompositeDilationPingPong1"))
+	};
+
+	// Greedy dispatch: consume DILATION_SIZE=2 steps (2 pixels each) until fewer than 2 pixels
+	// remain, then finish with DILATION_SIZE=1. Same total pixel reach as Size×DILATION_SIZE=1
+	// but up to half the number of GPU dispatches.
+	// Pixel tracks pixels consumed so far (stepping by 2); Iteration counts actual dispatches.
+	int32 Iteration = 0;
+	for (int32 Pixel = 0; Pixel < Size; Pixel += 2, ++Iteration)
+	{
+		const int32 StepSize = FMath::Min(Size - Pixel, 2);
+
+		// Determine read source: first iteration reads from input, subsequent alternate ping-pong.
+		FScreenPassTexture ReadTexture;
+		if (Iteration == 0)
+		{
+			ReadTexture = InputTexture;
+		}
+		else
+		{
+			const int32 ReadIndex = (Iteration - 1) % 2;
+			ReadTexture = FScreenPassTexture(PingPong[ReadIndex], FIntRect(FIntPoint::ZeroValue, TextureSize));
+		}
+
+		// Final iteration writes directly to Output; intermediate writes to ping-pong.
+		const bool bFinalIteration = (Pixel + StepSize >= Size);
+		FRDGTextureRef WriteTexture;
+		FScreenPassTextureViewportParameters OutputViewportParams;
+		if (bFinalIteration)
+		{
+			WriteTexture = Output.Texture;
+			OutputViewportParams = GetScreenPassTextureViewportParameters(FScreenPassTextureViewport(Output));
+		}
+		else
+		{
+			const int32 WriteIndex = Iteration % 2;
+			WriteTexture = PingPong[WriteIndex];
+			OutputViewportParams = GetScreenPassTextureViewportParameters(
+				FScreenPassTextureViewport(FScreenPassTexture(WriteTexture, FIntRect(FIntPoint::ZeroValue, TextureSize))));
+		}
+
+		FCompositePassDilationShader::FPermutationDomain PermutationVector;
+		PermutationVector.Set<FCompositePassDilationShader::FDilationSize>(StepSize);
+		PermutationVector.Set<FCompositePassDilationShader::FErode>(bErode);
+		TShaderMapRef<FCompositePassDilationShader> ComputeShader(GlobalShaderMap, PermutationVector);
+
+		FCompositePassDilationShader::FParameters* PassParameters = GraphBuilder.AllocParameters<FCompositePassDilationShader::FParameters>();
+		PassParameters->Input = GetScreenPassTextureInput(ReadTexture, TStaticSamplerState<SF_Point>::GetRHI());
+		PassParameters->Output = OutputViewportParams;
+		PassParameters->RWOutputTexture = GraphBuilder.CreateUAV(WriteTexture);
+		PassParameters->bInvertedAlpha = static_cast<uint32>(Metadata.bInvertedAlpha);
+		PassParameters->bOpacifyOutput = 0;
+		PassParameters->AlphaThreshold = AlphaThreshold;
+
+		// Dispatch over the write viewport so the final iteration always covers the full output.
+		const FIntPoint DispatchSize = bFinalIteration ? Output.ViewRect.Size() : TextureSize;
+		FComputeShaderUtils::AddPass(
+			GraphBuilder,
+			RDG_EVENT_NAME("Composite.Dilation %s (%dx%d) [pix:%d/%d step:%d]", bErode ? TEXT("Erode") : TEXT("Dilate"), TextureSize.X, TextureSize.Y, Pixel, Size, StepSize),
+			ERDGPassFlags::Compute,
+			ComputeShader,
+			PassParameters,
+			FComputeShaderUtils::GetGroupCount(DispatchSize, FCompositePassDilationShader::ThreadGroupSize)
+		);
+	}
+
+	return UE::CompositeCore::FPassTexture{ MoveTemp(Output), Metadata };
+}
+
+}  // namespace UE::Composite::Private
+
+UCompositePassDilation::UCompositePassDilation(const FObjectInitializer& ObjectInitializer)
+	: Super(ObjectInitializer)
+	, Mode(ECompositeDilationMode::Dilate)
+	, Size(1)
+{
+}
+
+UCompositePassDilation::~UCompositePassDilation() = default;
+
+bool UCompositePassDilation::GetIsActive() const
+{
+	return Super::GetIsActive() && Size >= 1;
+}
+
+FCompositeCorePassProxy* UCompositePassDilation::GetProxy(const UE::CompositeCore::FPassInputDecl& InputDecl, FCompositeTraversalContext& InContext, FSceneRenderingBulkObjectAllocator& InFrameAllocator) const
+{
+	using namespace UE::Composite::Private;
+
+	FCompositePassDilationProxy* Proxy = InFrameAllocator.Create<FCompositePassDilationProxy>(UE::CompositeCore::FPassInputDeclArray{ InputDecl });
+	Proxy->Size = Size;
+	Proxy->bErode = (Mode == ECompositeDilationMode::Erode);
+	Proxy->AlphaThreshold = AlphaThreshold;
+
+	return Proxy;
+}
diff --git a/Engine/Plugins/Compositing/Composite/Source/Composite/Private/Tests/CompositePassDilationTest.cpp b/Engine/Plugins/Compositing/Composite/Source/Composite/Private/Tests/CompositePassDilationTest.cpp
new file mode 100644
--- /dev/null
+++ b/Engine/Plugins/Compositing/Composite/Source/Composite/Private/Tests/CompositePassDilationTest.cpp
@@ -0,0 +1,338 @@
+// Copyright Epic Games, Inc. All Rights Reserved.
+
+#include "CoreTypes.h"
+#include "RenderGraphBuilder.h"
+#include "RenderGraphUtils.h"
+#include "RHI.h"
+#include "RHIGPUReadback.h"
+#include "ScreenPass.h"
+
+#include "Passes/CompositePassDilation.h"
+#include "Tests/CompositeTestHelpers.h"
+
+#if WITH_DEV_AUTOMATION_TESTS
+#include "Misc/AutomationTest.h"
+#endif
+
+namespace UE::Composite::Tests
+{
+
+#if WITH_DEV_AUTOMATION_TESTS
+
+BEGIN_SHADER_PARAMETER_STRUCT(FDilationTestUploadParameters, )
+	RDG_TEXTURE_ACCESS(InputTexture, ERHIAccess::CopyDest)
+END_SHADER_PARAMETER_STRUCT()
+
+namespace Dilation  // unique name prevents symbol collision in unity builds
+{
+using namespace UE::Composite::Private;
+using namespace UE::CompositeCore;
+
+// Fill a texture with vertical strips: left half = opaque red, right half = transparent black.
+// bInvertedAlpha=true  (UE convention): opaque=A=0, transparent=A=1.
+// bInvertedAlpha=false (standard):      opaque=A=1, transparent=A=0.
+void UploadVerticalStrip(FRDGBuilder& GraphBuilder, FRDGTextureRef Tex, bool bInvertedAlpha, const TCHAR* PassName)
+{
+	const int32 HalfWidth        = GTestViewSize.X / 2;
+	const float OpaqueAlpha      = bInvertedAlpha ? 0.0f : 1.0f;
+	const float TransparentAlpha = 1.0f - OpaqueAlpha;
+
+	// Build pixel data on the CPU before the pass. UpdateTexture2D handles the staging
+	// buffer internally; LockTexture2D is not safe on RDG textures without CPU-writable flags.
+	TArray<FLinearColor> PixelData;
+	PixelData.SetNumUninitialized(GTestViewSize.X * GTestViewSize.Y);
+	for (int32 Row = 0; Row < GTestViewSize.Y; ++Row)
+	{
+		for (int32 Col = 0; Col < GTestViewSize.X; ++Col)
+		{
+			PixelData[Row * GTestViewSize.X + Col] = (Col < HalfWidth)
+				? FLinearColor(1.0f, 0.0f, 0.0f, OpaqueAlpha)
+				: FLinearColor(0.0f, 0.0f, 0.0f, TransparentAlpha);
+		}
+	}
+
+	FDilationTestUploadParameters* Params = GraphBuilder.AllocParameters<FDilationTestUploadParameters>();
+	Params->InputTexture = Tex;
+
+	GraphBuilder.AddPass(
+		RDG_EVENT_NAME("DilationTest.Upload(%s)", PassName), Params,
+		ERDGPassFlags::Copy | ERDGPassFlags::NeverCull,
+		[Params, PixelData = MoveTemp(PixelData)](FRHICommandListImmediate& RHICmdList)
+		{
+			const uint32 SrcPitch = GTestViewSize.X * sizeof(FLinearColor);
+			RHICmdList.UpdateTexture2D(
+				Params->InputTexture->GetRHI(),
+				0,
+				FUpdateTextureRegion2D(0, 0, 0, 0, GTestViewSize.X, GTestViewSize.Y),
+				SrcPitch,
+				reinterpret_cast<const uint8*>(PixelData.GetData())
+			);
+		});
+}
+
+// Create textures, upload the vertical strip, run the proxy, and schedule a GPU readback.
+// Readback must be locked AFTER GraphBuilder.Execute() + BlockUntilGPUIdle() — use the
+// PostWork callback of DispatchAndWait.
+void RunAndSample(
+	FRDGBuilder& GraphBuilder, const FSceneView& View,
+	FCompositePassDilationProxy& Proxy, bool bInvertedAlpha,
+	FRHIGPUTextureReadback& Readback, const TCHAR* Tag)
+{
+	const FIntRect PixelRect(0, 0, GTestViewSize.X, GTestViewSize.Y);
+
+	const FRDGTextureDesc InputDesc = FRDGTextureDesc::Create2D(
+		GTestViewSize, PF_A32B32G32R32F, FClearValueBinding::Black, ETextureCreateFlags::ShaderResource);
+	const FRDGTextureDesc OutputDesc = FRDGTextureDesc::Create2D(
+		GTestViewSize, PF_A32B32G32R32F, FClearValueBinding::Black,
+		ETextureCreateFlags::RenderTargetable | ETextureCreateFlags::UAV);
+
+	FRDGTextureRef InputTex  = GraphBuilder.CreateTexture(InputDesc,  *FString::Printf(TEXT("%s.Input"),  Tag));
+	FRDGTextureRef OutputTex = GraphBuilder.CreateTexture(OutputDesc, *FString::Printf(TEXT("%s.Output"), Tag));
+	UploadVerticalStrip(GraphBuilder, InputTex, bInvertedAlpha, Tag);
+
+	FPassInputArray Inputs;
+	FPassInput& In = Inputs.GetArray().AddDefaulted_GetRef();
+	In.Texture              = FScreenPassTexture{ InputTex, PixelRect };
+	In.Metadata.Filter      = SF_Point;
+	In.Metadata.bInvertedAlpha = bInvertedAlpha;
+	Inputs.OverrideOutput = FScreenPassRenderTarget{
+		FScreenPassTexture{ OutputTex, PixelRect }, ERenderTargetLoadAction::ENoAction };
+
+	FPassContext PassContext;
+	PassContext.OutputViewRect = PixelRect;
+	Proxy.Add(GraphBuilder, View, Inputs, PassContext);
+
+	// Schedule the GPU→staging copy. The actual readback data is only valid after
+	// GraphBuilder.Execute() + BlockUntilGPUIdle() — call Readback.Lock() in PostWork.
+	AddEnqueueCopyPass(GraphBuilder, &Readback, OutputTex);
+}
+
+// Extract pixels at the given columns from the center row of a readback.
+// Must be called after BlockUntilGPUIdle() (i.e., from DispatchAndWait's PostWork callback).
+void ExtractColumns(FRHIGPUTextureReadback& Readback, const TArray<int32>& Cols, FLinearColor* Out)
+{
+	const int32 CenterRow = GTestViewSize.Y / 2;
+	int32 RowPitchInPixels = 0;
+	const FLinearColor* Pixels = static_cast<const FLinearColor*>(Readback.Lock(RowPitchInPixels));
+	for (int32 i = 0; i < Cols.Num(); ++i)
+	{
+		Out[i] = Pixels[CenterRow * RowPitchInPixels + Cols[i]];
+	}
+	Readback.Unlock();
+}
+
+}  // namespace Dilation
+
+// ============================================================================
+// Test 1: Size=1 — dilation expands by 1 pixel, erosion shrinks by 1 pixel
+// ============================================================================
+//
+// Input: cols 0..31 = opaque red (A=0), cols 32..63 = transparent black (A=1).
+//
+// Dilation (Size=1): only the first transparent column (32) becomes opaque.
+//   col 30: A=0 (opaque, far from boundary)
+//   col 31: A=0 (boundary opaque — unchanged)
+//   col 32: A=0 (first transparent, dilated by left neighbor)
+//   col 33: A=1 (still transparent — one step beyond dilation reach)
+//
+// Erosion (Size=1): only the last opaque column (31) becomes transparent.
+//   col 30: A=0 (opaque, one step in from boundary)
+//   col 31: A=1 (boundary opaque — eroded by right transparent neighbor)
+//   col 32: A=1 (transparent — unchanged)
+//   col 33: A=1 (transparent — unchanged)
+
+IMPLEMENT_SIMPLE_AUTOMATION_TEST(FCompositeDilationPassTest, "Composite.UnitTests.DilationPass",
+	EAutomationTestFlags_ApplicationContextMask | EAutomationTestFlags::EngineFilter | EAutomationTestFlags::NonNullRHI)
+
+bool FCompositeDilationPassTest::RunTest(const FString& Parameter)
+{
+	using namespace Dilation;
+
+	TArray<FLinearColor> Results;
+	Results.SetNumZeroed(8);
+
+	FCompositePassDilationProxy DilateProxy(GetDefaultInputDeclArray());
+	DilateProxy.Size   = 1;
+	DilateProxy.bErode = false;
+
+	FCompositePassDilationProxy ErodeProxy(GetDefaultInputDeclArray());
+	ErodeProxy.Size   = 1;
+	ErodeProxy.bErode = true;
+
+	FRHIGPUTextureReadback DilateRB(TEXT("DilateS1")), ErodeRB(TEXT("ErodeS1"));
+
+	FCompositeTestView TestView;
+	DispatchAndWait(TestView.Get(),
+		[&, DProxy = MoveTemp(DilateProxy), EProxy = MoveTemp(ErodeProxy)]
+		(FRDGBuilder& GraphBuilder, const FSceneView& View) mutable
+		{
+			RunAndSample(GraphBuilder, View, DProxy, true, DilateRB, TEXT("DilateS1"));
+			RunAndSample(GraphBuilder, View, EProxy, true, ErodeRB,  TEXT("ErodeS1"));
+		},
+		[&](FRHICommandListImmediate&)
+		{
+			ExtractColumns(DilateRB, {30, 31, 32, 33}, &Results[0]);
+			ExtractColumns(ErodeRB,  {30, 31, 32, 33}, &Results[4]);
+		});
+
+	constexpr float Tol = 1e-3f;
+
+	// Dilation Size=1
+	UTEST_TRUE("Dilate S1 col30: stays opaque",       Results[0].A < Tol);
+	UTEST_TRUE("Dilate S1 col31: stays opaque",       Results[1].A < Tol);
+	UTEST_TRUE("Dilate S1 col32: becomes opaque",     Results[2].A < Tol);
+	UTEST_TRUE("Dilate S1 col33: stays transparent",  Results[3].A > 1.0f - Tol);
+
+	// Erosion Size=1
+	UTEST_TRUE("Erode S1 col30: stays opaque",        Results[4].A < Tol);
+	UTEST_TRUE("Erode S1 col31: becomes transparent", Results[5].A > 1.0f - Tol);
+	UTEST_TRUE("Erode S1 col32: stays transparent",   Results[6].A > 1.0f - Tol);
+	UTEST_TRUE("Erode S1 col33: stays transparent",   Results[7].A > 1.0f - Tol);
+
+	return true;
+}
+
+// ============================================================================
+// Test 2: Size=4 — dilation expands by 4 pixels, erosion shrinks by 4 pixels
+// ============================================================================
+//
+// Same input layout. Size=4 is dispatched as 2 × DILATION_SIZE=2 passes (greedy dispatch),
+// covering 4 pixels total across 2 GPU iterations (Input → PingPong[0] → Output).
+//
+// Dilation (Size=4): transparent columns 32–35 become opaque.
+//   col 30: A=0 (unchanged)
+//   col 32: A=0 (within reach of iteration 1 — 2-pixel step)
+//   col 35: A=0 (within reach of iteration 2 — stresses the full 2-pass chain)
+//   col 36: A=1 (still transparent — one beyond reach)
+//
+// Erosion (Size=4): opaque columns 31–28 become transparent.
+//   col 27: A=0 (unchanged — 4 pixels in from boundary, survives)
+//   col 28: A=1 (within reach of iteration 2 — stresses the full 2-pass chain)
+//   col 31: A=1 (within reach of iteration 1 — 2-pixel step)
+//   col 32: A=1 (unchanged)
+
+IMPLEMENT_SIMPLE_AUTOMATION_TEST(FCompositeDilationPassSize4Test, "Composite.UnitTests.DilationPassSize4",
+	EAutomationTestFlags_ApplicationContextMask | EAutomationTestFlags::EngineFilter | EAutomationTestFlags::NonNullRHI)
+
+bool FCompositeDilationPassSize4Test::RunTest(const FString& Parameter)
+{
+	using namespace Dilation;
+
+	TArray<FLinearColor> Results;
+	Results.SetNumZeroed(8);
+
+	FCompositePassDilationProxy DilateProxy(GetDefaultInputDeclArray());
+	DilateProxy.Size   = 4;
+	DilateProxy.bErode = false;
+
+	FCompositePassDilationProxy ErodeProxy(GetDefaultInputDeclArray());
+	ErodeProxy.Size   = 4;
+	ErodeProxy.bErode = true;
+
+	FRHIGPUTextureReadback DilateRB(TEXT("DilateS4")), ErodeRB(TEXT("ErodeS4"));
+
+	FCompositeTestView TestView;
+	DispatchAndWait(TestView.Get(),
+		[&, DProxy = MoveTemp(DilateProxy), EProxy = MoveTemp(ErodeProxy)]
+		(FRDGBuilder& GraphBuilder, const FSceneView& View) mutable
+		{
+			RunAndSample(GraphBuilder, View, DProxy, true, DilateRB, TEXT("DilateS4"));
+			RunAndSample(GraphBuilder, View, EProxy, true, ErodeRB,  TEXT("ErodeS4"));
+		},
+		[&](FRHICommandListImmediate&)
+		{
+			ExtractColumns(DilateRB, {30, 32, 35, 36}, &Results[0]);
+			ExtractColumns(ErodeRB,  {27, 28, 31, 32}, &Results[4]);
+		});
+
+	constexpr float Tol = 1e-3f;
+
+	// Dilation Size=4
+	UTEST_TRUE("Dilate S4 col30: stays opaque",       Results[0].A < Tol);
+	UTEST_TRUE("Dilate S4 col32: becomes opaque",     Results[1].A < Tol);
+	UTEST_TRUE("Dilate S4 col35: becomes opaque",     Results[2].A < Tol);
+	UTEST_TRUE("Dilate S4 col36: stays transparent",  Results[3].A > 1.0f - Tol);
+
+	// Erosion Size=4
+	UTEST_TRUE("Erode S4 col27: stays opaque",        Results[4].A < Tol);
+	UTEST_TRUE("Erode S4 col28: becomes transparent", Results[5].A > 1.0f - Tol);
+	UTEST_TRUE("Erode S4 col31: becomes transparent", Results[6].A > 1.0f - Tol);
+	UTEST_TRUE("Erode S4 col32: stays transparent",   Results[7].A > 1.0f - Tol);
+
+	return true;
+}
+
+// ============================================================================
+// Test 3: Standard alpha (bInvertedAlpha=false) — Size=1 dilation and erosion
+// ============================================================================
+//
+// Same boundary geometry but standard alpha convention: A=1=opaque, A=0=transparent.
+// Assertions are the mirror of Test 1 — opaque pixels have A near 1, transparent near 0.
+//
+// Dilation (Size=1):
+//   col 30: A=1 (opaque, far from boundary)
+//   col 31: A=1 (boundary opaque — unchanged)
+//   col 32: A=1 (first transparent, dilated by left neighbor)
+//   col 33: A=0 (still transparent — one step beyond dilation reach)
+//
+// Erosion (Size=1):
+//   col 30: A=1 (opaque, one step in from boundary)
+//   col 31: A=0 (boundary opaque — eroded by right transparent neighbor)
+//   col 32: A=0 (transparent — unchanged)
+//   col 33: A=0 (transparent — unchanged)
+
+IMPLEMENT_SIMPLE_AUTOMATION_TEST(FCompositeDilationPassStdAlphaTest, "Composite.UnitTests.DilationPassStdAlpha",
+	EAutomationTestFlags_ApplicationContextMask | EAutomationTestFlags::EngineFilter | EAutomationTestFlags::NonNullRHI)
+
+bool FCompositeDilationPassStdAlphaTest::RunTest(const FString& Parameter)
+{
+	using namespace Dilation;
+
+	TArray<FLinearColor> Results;
+	Results.SetNumZeroed(8);
+
+	FCompositePassDilationProxy DilateProxy(GetDefaultInputDeclArray());
+	DilateProxy.Size   = 1;
+	DilateProxy.bErode = false;
+
+	FCompositePassDilationProxy ErodeProxy(GetDefaultInputDeclArray());
+	ErodeProxy.Size   = 1;
+	ErodeProxy.bErode = true;
+
+	FRHIGPUTextureReadback DilateRB(TEXT("DilateStd")), ErodeRB(TEXT("ErodeStd"));
+
+	FCompositeTestView TestView;
+	DispatchAndWait(TestView.Get(),
+		[&, DProxy = MoveTemp(DilateProxy), EProxy = MoveTemp(ErodeProxy)]
+		(FRDGBuilder& GraphBuilder, const FSceneView& View) mutable
+		{
+			RunAndSample(GraphBuilder, View, DProxy, false, DilateRB, TEXT("DilateStd"));
+			RunAndSample(GraphBuilder, View, EProxy, false, ErodeRB,  TEXT("ErodeStd"));
+		},
+		[&](FRHICommandListImmediate&)
+		{
+			ExtractColumns(DilateRB, {30, 31, 32, 33}, &Results[0]);
+			ExtractColumns(ErodeRB,  {30, 31, 32, 33}, &Results[4]);
+		});
+
+	constexpr float Tol = 1e-3f;
+
+	// Standard alpha: opaque = A > 1-Tol, transparent = A < Tol
+	// Dilation Size=1
+	UTEST_TRUE("Dilate StdAlpha col30: stays opaque",       Results[0].A > 1.0f - Tol);
+	UTEST_TRUE("Dilate StdAlpha col31: stays opaque",       Results[1].A > 1.0f - Tol);
+	UTEST_TRUE("Dilate StdAlpha col32: becomes opaque",     Results[2].A > 1.0f - Tol);
+	UTEST_TRUE("Dilate StdAlpha col33: stays transparent",  Results[3].A < Tol);
+
+	// Erosion Size=1
+	UTEST_TRUE("Erode StdAlpha col30: stays opaque",        Results[4].A > 1.0f - Tol);
+	UTEST_TRUE("Erode StdAlpha col31: becomes transparent", Results[5].A < Tol);
+	UTEST_TRUE("Erode StdAlpha col32: stays transparent",   Results[6].A < Tol);
+	UTEST_TRUE("Erode StdAlpha col33: stays transparent",   Results[7].A < Tol);
+
+	return true;
+}
+
+#endif // WITH_DEV_AUTOMATION_TESTS
+
+} // namespace UE::Composite::Tests
diff --git a/Engine/Plugins/Compositing/Composite/Source/Composite/Public/Passes/CompositePassDilation.h b/Engine/Plugins/Compositing/Composite/Source/Composite/Public/Passes/CompositePassDilation.h
new file mode 100644
--- /dev/null
+++ b/Engine/Plugins/Compositing/Composite/Source/Composite/Public/Passes/CompositePassDilation.h
@@ -0,0 +1,83 @@
+// Copyright Epic Games, Inc. All Rights Reserved.
+
+#pragma once
+
+#include "CompositePassBase.h"
+#include "Passes/CompositeCorePassProxy.h"
+
+#include "CompositePassDilation.generated.h"
+
+#define UE_API COMPOSITE_API
+
+/** Dilation mode type. */
+UENUM(BlueprintType)
+enum class ECompositeDilationMode : uint8
+{
+	Dilate,
+	Erode,
+};
+
+/** Morphological dilation/erosion composite pass. */
+UCLASS(MinimalAPI, BlueprintType, Blueprintable, EditInlineNew, CollapseCategories, meta = (DisplayName = "Dilation Pass"))
+class UCompositePassDilation : public UCompositePassBase
+{
+	GENERATED_BODY()
+
+public:
+	/** Constructor */
+	UE_API UCompositePassDilation(const FObjectInitializer& ObjectInitializer);
+
+	/** Destructor */
+	UE_API ~UCompositePassDilation();
+
+	UE_API virtual bool GetIsActive() const override;
+
+	UE_API virtual FCompositeCorePassProxy* GetProxy(const UE::CompositeCore::FPassInputDecl& InputDecl, FCompositeTraversalContext& InContext, FSceneRenderingBulkObjectAllocator& InFrameAllocator) const override;
+
+public:
+	/** Dilation mode: expand or shrink opaque regions. */
+	UPROPERTY(EditAnywhere, BlueprintReadWrite, Interp, Category = "Composite")
+	ECompositeDilationMode Mode;
+
+	/** Size of the dilation/erosion in pixels. */
+	UPROPERTY(EditAnywhere, BlueprintReadWrite, Interp, Category = "Composite", meta = (ClampMin = "1", UIMin = "1", UIMax = "20"))
+	int32 Size = 1;
+
+	/** Opacity threshold below which pixels are treated as fully transparent. In standard opacity space (0=transparent, 1=opaque). */
+	UPROPERTY(EditAnywhere, BlueprintReadWrite, Category = "Composite", meta = (ClampMin = "0.0", ClampMax = "0.5"))
+	float AlphaThreshold = 0.0001f;
+};
+
+#undef UE_API
+
+namespace UE
+{
+	namespace Composite
+	{
+		namespace Private
+		{
+			using namespace CompositeCore;
+
+			/** Render-thread proxy for the morphological dilation/erosion pass. */
+			class FCompositePassDilationProxy : public FCompositeCorePassProxy
+			{
+			public:
+				IMPLEMENT_COMPOSITE_PASS(FCompositePassDilationProxy);
+
+				using FCompositeCorePassProxy::FCompositeCorePassProxy;
+
+				/** Adds ceil(Size/2) RDG compute passes covering Size pixels of dilation or erosion and returns the result.
+				 *  Each pass uses a DILATION_SIZE=2 kernel (2 pixels/dispatch) where possible, falling back to
+				 *  DILATION_SIZE=1 for an odd remainder. Intermediate results ping-pong through internal textures. */
+				FPassTexture Add(FRDGBuilder& GraphBuilder, const FSceneView& InView, const FPassInputArray& Inputs, const FPassContext& PassContext) const override;
+
+				/** Total pixel reach of the dilation/erosion (dispatched in greedy 2-pixel steps). */
+				int32 Size = 1;
+				/** When true, runs erosion (shrink opaque); when false, runs dilation (expand opaque). */
+				bool bErode = false;
+				/** Opacity threshold - pixels at or below this (in standard opacity space) are treated as fully transparent. */
+				float AlphaThreshold = 0.0001f;
+			};
+		}
+	}
+}
diff --git a/Engine/Plugins/Compositing/CompositeCore/Shaders/Private/CompositeCoreDilate.usf b/Engine/Plugins/Compositing/CompositeCore/Shaders/Private/CompositeCoreDilate.usf
--- a/Engine/Plugins/Compositing/CompositeCore/Shaders/Private/CompositeCoreDilate.usf
+++ b/Engine/Plugins/Compositing/CompositeCore/Shaders/Private/CompositeCoreDilate.usf
@@ -1,15 +1,19 @@
 // Copyright Epic Games, Inc. All Rights Reserved.
 
 #include "/Engine/Private/Common.ush"
+#include "/Engine/Private/ScreenPass.ush"
 
 #ifndef TILE_SIZE
     #error TILE_SIZE must be defined
 #endif
 
+#ifndef MORPHOLOGY_ERODE
+    #define MORPHOLOGY_ERODE 0
+#endif
+
 #define KERNEL_COUNT ((1 + 2 * DILATION_SIZE) * (1 + 2 * DILATION_SIZE) - 1)
-#define NEARLY_TRANSLUCENT_ALPHA 0.9999999f // Alpha is inverted opacity
 
-// Number of texels arround the group tile for scene color.
+// Number of texels around the group tile for scene color.
 #define LDS_TILE_BORDER_SIZE DILATION_SIZE
 
 // Width in texels of the depth tile cached into LDS.
@@ -20,10 +24,24 @@
 
 #define THREADGROUP_TOTAL (TILE_SIZE * TILE_SIZE)
 
-Texture2D InputTexture;
+Texture2D Input_Texture;
+SamplerState Input_Sampler;
+SCREEN_PASS_TEXTURE_VIEWPORT(Input)
+SCREEN_PASS_TEXTURE_VIEWPORT(Output)
 RWTexture2D<float4> RWOutputTexture;
-uint2 Dimensions;
+uint bInvertedAlpha;
 uint bOpacifyOutput;
+// Opacity at or below this threshold is treated as fully transparent for morphological operations.
+// In standard opacity space (0=transparent, 1=opaque): opacity < AlphaThreshold means "treat as transparent".
+float AlphaThreshold;
+
+// Normalize or restore alpha between inverted UE convention (A=0 opaque, A=1 transparent)
+// and standard opacity (0 transparent, 1 opaque). Self-inverse: apply twice to round-trip.
+// Defined inline — CompositeCommon.ush lives in the Composite plugin and is not accessible here.
+float InvertAlpha(float InAlpha, uint bInvert)
+{
+    return select(bInvert > 0u, 1.0f - InAlpha, InAlpha);
+}
 
 #if DILATION_SIZE == 1
 static const int2 kOffsets[KERNEL_COUNT] = {
@@ -81,16 +99,15 @@ float4 SampleCachedColors(in uint2 GroupThreadId, int2 PixelOffset = int2(0,0))
 
 float SafeRcp(float x)
 {
-	return x > 0.0 ? rcp(x) : 0.0;
+    return x > 0.0 ? rcp(x) : 0.0;
 }
 
+// Unpremultiply and force full opacity. Operates in standard opacity space (a = 1 = opaque).
+// Only used by CompositeCore's holdout path (bOpacifyOutput=1).
 float4 Opacify(float4 InColor)
 {
-    float RegularAlpha = saturate(1.0 - InColor.a);
-    // Unpremultiply color
-    InColor.xyz *= SafeRcp(RegularAlpha);
-    // Return opaque color, maintaining UE's inverted alpha convention
-	return float4(InColor.xyz, 0.0);
+    InColor.xyz *= SafeRcp(saturate(InColor.a));
+    return float4(InColor.xyz, 1.0f);
 }
 
 [numthreads(TILE_SIZE, TILE_SIZE, 1)]
@@ -102,8 +119,12 @@ void MainCS(
 {
     // Init LDS
     {
-        // Global texel coordinate of the active group (corner) offset by the LDS border
-        uint2 GroupTexelOffset = GroupId * uint2(TILE_SIZE, TILE_SIZE) - LDS_TILE_BORDER_SIZE;
+        // GroupTexelOffset must be signed: when GroupId.x or GroupId.y is 0 the halo extends one
+        // pixel outside the viewport (position -1), which uint arithmetic would wrap to 0xFFFFFFFF.
+        // The subsequent min-clamp would then map to ViewportSize-1 (the opposite edge) rather
+        // than 0, loading pixels from the wrong corner and creating phantom opaque halos that look
+        // like diagonal propagation artefacts at the viewport boundary.
+        int2 GroupTexelOffsetSigned = int2(GroupId) * int2(TILE_SIZE, TILE_SIZE) - LDS_TILE_BORDER_SIZE;
 
         // To account for the dilation margin, each warp thread should load extra texels.
         // We proceed linearly through the tile (with borders) and each warp thread loads
@@ -115,15 +136,17 @@ void MainCS(
         UNROLL
         for (uint i = 0; i < LoadCount; i++)
         {
-            uint2 TexelLocation =  GroupTexelOffset + uint2(
-                LinearGroupThreadId % LDS_TILE_WIDTH,
-                LinearGroupThreadId / LDS_TILE_WIDTH);
+            // Compute the signed viewport-relative texel position and clamp to [0, ViewportSize-1].
+            int2 TexelLocationSigned = GroupTexelOffsetSigned + int2(
+                int(LinearGroupThreadId % LDS_TILE_WIDTH),
+                int(LinearGroupThreadId / LDS_TILE_WIDTH));
+            TexelLocationSigned = clamp(TexelLocationSigned, int2(0, 0), int2(Input_ViewportSize) - 1);
 
             if ((LinearGroupThreadId < LDS_ARRAY_SIZE) ||
                 (i != LoadCount - 1) ||
                 (LDS_ARRAY_SIZE % THREADGROUP_TOTAL) == 0)
             {
-                SharedColor[LinearGroupThreadId] = InputTexture[TexelLocation];
+                SharedColor[LinearGroupThreadId] = Input_Texture[uint2(TexelLocationSigned) + uint2(Input_ViewportMin)];
             }
             LinearGroupThreadId += THREADGROUP_TOTAL;
         }
@@ -132,37 +155,61 @@ void MainCS(
     GroupMemoryBarrierWithGroupSync();
 
     // Discard out-of-bound texels
-    if (any(DispatchThreadId >= Dimensions))
+    if (any(DispatchThreadId >= uint2(Output_ViewportSize)))
     {
         return;
     }
 
     float4 TexelColor = SampleCachedColors(GroupThreadId);
-    
+    // Normalize to standard opacity (0=transparent, 1=opaque) for all downstream logic.
+    TexelColor.a = InvertAlpha(TexelColor.a, bInvertedAlpha);
+
 #if DILATION_SIZE > 0
-	// Apply (gather) dilation on any translucent or nearly translucent texel
+#if MORPHOLOGY_ERODE
+    // Erosion: if texel is nearly opaque and any neighbor is nearly transparent, make texel transparent.
+    BRANCH
+    if (TexelColor.a > 1.0f - AlphaThreshold)
+    {
+        UNROLL
+        for (int i = 0; i < KERNEL_COUNT; i++)
+        {
+            float4 NeighborColor = SampleCachedColors(GroupThreadId, kOffsets[i]);
+            if (InvertAlpha(NeighborColor.a, bInvertedAlpha) < AlphaThreshold)
+            {
+                TexelColor = float4(0.0f, 0.0f, 0.0f, 0.0f); // transparent black (standard opacity)
+                break;
+            }
+        }
+    }
+#else
+    // Dilation: if texel is nearly transparent and any neighbor is more opaque, take neighbor color.
     BRANCH
-    if(TexelColor.a > NEARLY_TRANSLUCENT_ALPHA)
+    if (TexelColor.a < AlphaThreshold)
     {
         UNROLL
-        for(int i=0; i<KERNEL_COUNT; i++)
+        for (int i = 0; i < KERNEL_COUNT; i++)
         {
             float4 OffsetColor = SampleCachedColors(GroupThreadId, kOffsets[i]);
+            float OffsetOpacity = InvertAlpha(OffsetColor.a, bInvertedAlpha);
 
-            if(OffsetColor.a < TexelColor.a)
+            if (OffsetOpacity > TexelColor.a)
             {
-                TexelColor = OffsetColor;
+                TexelColor = float4(OffsetColor.rgb, OffsetOpacity);
                 break; // We break as soon as we find a more opaque color
             }
         }
     }
-#endif
+#endif // MORPHOLOGY_ERODE
+#endif // DILATION_SIZE > 0
 
-	// Condtionally & fully opacify non-translucent colors for the over blend.
-	if (bOpacifyOutput && TexelColor.a < NEARLY_TRANSLUCENT_ALPHA)
-	{
-		TexelColor = Opacify(TexelColor);
-	}
+    // Conditionally & fully opacify non-translucent colors for the over blend.
+    // Only used by CompositeCore's holdout path (bOpacifyOutput=1).
+    if (bOpacifyOutput && TexelColor.a > 1.0f - AlphaThreshold)
+    {
+        TexelColor = Opacify(TexelColor);
+    }
 
-	RWOutputTexture[DispatchThreadId] = TexelColor;
+    // Restore original alpha convention before writing.
+    TexelColor.a = InvertAlpha(TexelColor.a, bInvertedAlpha);
+    RWOutputTexture[Output_ViewportMin + DispatchThreadId] = TexelColor;
 }
diff --git a/Engine/Plugins/Compositing/CompositeCore/Source/CompositeCore/Private/Passes/CompositeCorePassDilate.cpp b/Engine/Plugins/Compositing/CompositeCore/Source/CompositeCore/Private/Passes/CompositeCorePassDilate.cpp
--- a/Engine/Plugins/Compositing/CompositeCore/Source/CompositeCore/Private/Passes/CompositeCorePassDilate.cpp
+++ b/Engine/Plugins/Compositing/CompositeCore/Source/CompositeCore/Private/Passes/CompositeCorePassDilate.cpp
@@ -6,6 +6,7 @@
 #include "RenderGraphBuilder.h"
 #include "RenderGraphUtils.h"
 #include "PixelShaderUtils.h"
+#include "ScreenPass.h"
 #include "ProfilingDebugging/RealtimeGPUProfiler.h"
 
 DECLARE_GPU_STAT_NAMED(FCompositeCoreDilate, TEXT("CompositeCore.Dilate"));
@@ -19,10 +20,12 @@ class FCompositeCoreDilateShader : public FGlobalShader
 	using FPermutationDomain = TShaderPermutationDomain<FDilationSize>;
 
 	BEGIN_SHADER_PARAMETER_STRUCT(FParameters, )
-		SHADER_PARAMETER_RDG_TEXTURE(Texture2D, InputTexture)
+		SHADER_PARAMETER_STRUCT(FScreenPassTextureInput, Input)
+		SHADER_PARAMETER_STRUCT(FScreenPassTextureViewportParameters, Output)
 		SHADER_PARAMETER_RDG_TEXTURE_UAV(RWTexture2D<float4>, RWOutputTexture)
-		SHADER_PARAMETER(FIntPoint, Dimensions)
+		SHADER_PARAMETER(uint32, bInvertedAlpha)
 		SHADER_PARAMETER(uint32, bOpacifyOutput)
+		SHADER_PARAMETER(float, AlphaThreshold)
 	END_SHADER_PARAMETER_STRUCT()
 
 	static inline void ModifyCompilationEnvironment(const FGlobalShaderPermutationParameters& Parameters, FShaderCompilerEnvironment& OutEnvironment)
@@ -46,14 +49,24 @@ namespace UE
 			{
 				RDG_EVENT_SCOPE_STAT(GraphBuilder, FCompositeCoreDilate, "CompositeCore.Dilate");
 
+				checkf(Input->Desc.Extent == Output->Desc.Extent,
+					TEXT("AddDilatePass: Input (%dx%d) and Output (%dx%d) must have matching extents"),
+					Input->Desc.Extent.X, Input->Desc.Extent.Y, Output->Desc.Extent.X, Output->Desc.Extent.Y);
+
 				FGlobalShaderMap* GlobalShaderMap = GetGlobalShaderMap(FeatureLevel);
 				const FIntPoint TextureSize = Input->Desc.Extent;
 
+				// CompositeCore's dilation texture always fills the entire extent with no viewport offset.
+				FScreenPassTexture InputScreenPass(Input, FIntRect(FIntPoint::ZeroValue, TextureSize));
+
 				FCompositeCoreDilateShader::FParameters* PassParameters = GraphBuilder.AllocParameters<FCompositeCoreDilateShader::FParameters>();
-				PassParameters->InputTexture = Input;
+				PassParameters->Input = GetScreenPassTextureInput(InputScreenPass, TStaticSamplerState<SF_Point>::GetRHI());
+				FScreenPassTexture OutputScreenPass(Output, FIntRect(FIntPoint::ZeroValue, TextureSize));
+				PassParameters->Output = GetScreenPassTextureViewportParameters(FScreenPassTextureViewport(OutputScreenPass));
 				PassParameters->RWOutputTexture = GraphBuilder.CreateUAV(Output);
-				PassParameters->Dimensions = TextureSize;
+				PassParameters->bInvertedAlpha = 1; // CompositeCore always uses inverted alpha (A=0 opaque, A=1 transparent)
 				PassParameters->bOpacifyOutput = PassInputs.bOpacifyOutput;
+				PassParameters->AlphaThreshold = PassInputs.AlphaThreshold;
 
 				FCompositeCoreDilateShader::FPermutationDomain PermutationVector;
 				PermutationVector.Set<FCompositeCoreDilateShader::FDilationSize>(FMath::Clamp(PassInputs.DilationSize, 0, 2));
@@ -88,4 +101,3 @@ namespace UE
 		}
 	}
 }
-
diff --git a/Engine/Plugins/Compositing/CompositeCore/Source/CompositeCore/Private/Passes/CompositeCorePassDilate.h b/Engine/Plugins/Compositing/CompositeCore/Source/CompositeCore/Private/Passes/CompositeCorePassDilate.h
--- a/Engine/Plugins/Compositing/CompositeCore/Source/CompositeCore/Private/Passes/CompositeCorePassDilate.h
+++ b/Engine/Plugins/Compositing/CompositeCore/Source/CompositeCore/Private/Passes/CompositeCorePassDilate.h
@@ -13,11 +13,14 @@ namespace UE
 		{
 			struct FDilateInputs
 			{
-				/** Size of the dilation step. WARNING: Not currently implemented. */
+				/** Dilation kernel radius in pixels (0 = no-op, 1 = 3×3 neighborhood, 2 = 5×5 neighborhood). Clamped to [0, 2]. */
 				int32 DilationSize = 1;
-				
+
 				/** Opacify the pass output to solid colors. */
 				bool bOpacifyOutput = true;
+
+				/** Opacity threshold - pixels at or below this (standard opacity space) are treated as fully transparent. */
+				float AlphaThreshold = 0.0000001f;
 			};
 
 			/**
