Composure: Minor fix to preserve original resolution through recently added dilation pass.

#Synced-CL 52419094

diff --git a/Engine/Plugins/Compositing/Composite/Source/Composite/Private/Passes/CompositePassDilation.cpp b/Engine/Plugins/Compositing/Composite/Source/Composite/Private/Passes/CompositePassDilation.cpp
--- a/Engine/Plugins/Compositing/Composite/Source/Composite/Private/Passes/CompositePassDilation.cpp
+++ b/Engine/Plugins/Compositing/Composite/Source/Composite/Private/Passes/CompositePassDilation.cpp
@@ -67,10 +67,14 @@ FPassTexture FCompositePassDilationProxy::Add(FRDGBuilder& GraphBuilder, const F
 		{
 			Output = {};
 		}
+		else if (!ensureMsgf(Output.ViewRect.Size() == InputTexture.ViewRect.Size(), TEXT("CompositePassDilation: OverrideOutput viewport (%dx%d) does not match input viewport (%dx%d); ignoring override and allocating an internal output instead."), Output.ViewRect.Size().X, Output.ViewRect.Size().Y, InputTexture.ViewRect.Size().X, InputTexture.ViewRect.Size().Y))
+		{
+			Output = {};
+		}
 	}
 	if (!Output.IsValid())
 	{
-		Output = CreateOutputRenderTarget(GraphBuilder, InView, PassContext.OutputViewRect, InputTexture.Texture->Desc, TEXT("CompositeDilationOutput"));
+		Output = CreateOutputRenderTarget(GraphBuilder, InView, InputTexture.ViewRect, InputTexture.Texture->Desc, TEXT("CompositeDilationOutput"));
 	}
 
 	// Input and Output viewports must be the same size: the dilation kernel operates at a fixed
