Composite: Additional dilation pass fix for compute-to-graphics flicker by routing CS output through UAV-only intermediates.

- Final dilation CS now writes to a UAV/SRV-only intermediate, then AddCopyTexturePass bridges to Output.Texture; mirroring the established engine pattern
- Same fix applied to the keyer CS output (KeyerComputeTexture)
- Removed unnecessary TexCreate_RenderTargetable for both passes

#Synced-CL 52419094

diff --git a/Engine/Plugins/Compositing/Composite/Source/Composite/Private/Passes/CompositePassColorKeyer.cpp b/Engine/Plugins/Compositing/Composite/Source/Composite/Private/Passes/CompositePassColorKeyer.cpp
--- a/Engine/Plugins/Compositing/Composite/Source/Composite/Private/Passes/CompositePassColorKeyer.cpp
+++ b/Engine/Plugins/Compositing/Composite/Source/Composite/Private/Passes/CompositePassColorKeyer.cpp
@@ -5,6 +5,7 @@
 #include "Engine/Texture.h"
 #include "Passes/CompositeCorePassProxy.h"
 #include "PixelShaderUtils.h"
+#include "RenderGraphUtils.h"
 #include "ScreenPass.h"
 #include "SystemTextures.h"
 #include "TextureResource.h"
@@ -82,7 +83,7 @@ static FScreenPassTexture AddColorKeyerDenoisePass(FRDGBuilder& GraphBuilder, co
 	FGlobalShaderMap* GlobalShaderMap = GetGlobalShaderMap(FeatureLevel);
 	const FIntPoint TextureSize = Input.ViewRect.Size();
 
-	const FRDGTextureDesc OutputDesc = FRDGTextureDesc::Create2D(TextureSize, PF_FloatRGBA, FClearValueBinding::Black, TexCreate_UAV | TexCreate_ShaderResource | TexCreate_RenderTargetable);
+	const FRDGTextureDesc OutputDesc = FRDGTextureDesc::Create2D(TextureSize, PF_FloatRGBA, FClearValueBinding::Black, TexCreate_UAV | TexCreate_ShaderResource);
 	FScreenPassRenderTarget Output = FScreenPassRenderTarget(GraphBuilder.CreateTexture(OutputDesc, TEXT("CompositeTextureDenoised")), FIntRect(FIntPoint::ZeroValue, TextureSize), ERenderTargetLoadAction::ENoAction);
 
 	FCompositePassDenoiseShader::FParameters* PassParameters = GraphBuilder.AllocParameters<FCompositePassDenoiseShader::FParameters>();
@@ -130,16 +131,29 @@ namespace UE::Composite::Private
 
 		FGlobalShaderMap* GlobalShaderMap = GetGlobalShaderMap(InView.GetFeatureLevel());
 
-		const FScreenPassTextureViewportParameters OutputParameters = GetScreenPassTextureViewportParameters(FScreenPassTextureViewport(Output));
+		/** Write the keyer CS result to a UAV-only intermediate rather than directly to Output.Texture.
+		 *  Output.Texture is RenderTargetable; D3D12 state tracking for UAV writes onto RenderTargetable
+		 *  textures can race with the RHI thread, corrupting the subsequent graphics SRV read in the
+		 *  merge pass. See CompositePassDilation.cpp and GeometryMaskPostProcess_Blur.cpp for the pattern. */
+		const FIntPoint OutputSize = Output.ViewRect.Size();
+		const FRDGTextureDesc KeyerComputeDesc = FRDGTextureDesc::Create2D(OutputSize, Output.Texture->Desc.Format, FClearValueBinding::Black, TexCreate_UAV | TexCreate_ShaderResource);
+		FRDGTextureRef KeyerComputeTexture = GraphBuilder.CreateTexture(KeyerComputeDesc, TEXT("CompositeColorKeyerCS"));
+
+		/** Zero-based viewport matching the CS dispatch and UAV write coordinates.
+		 *  Both OutputParameters and SvPositionToInputTextureUV must use the same viewport
+		 *  so that the TexelPosition -> ViewportUV basis is consistent with dispatch thread IDs. */
+		const FScreenPassTextureViewport KeyerViewport(FScreenPassTexture(KeyerComputeTexture, FIntRect(FIntPoint::ZeroValue, OutputSize)));
+
+		const FScreenPassTextureViewportParameters OutputParameters = GetScreenPassTextureViewportParameters(KeyerViewport);
 
 		FRHISamplerState* InputSampler = Metadata.GetSamplerState();
 
 		FCompositePassColorKeyerShader::FParameters* PassParameters = GraphBuilder.AllocParameters<FCompositePassColorKeyerShader::FParameters>();
 		PassParameters->Input = GetScreenPassTextureInput(Input, InputSampler);
-		PassParameters->RWOutputTexture = GraphBuilder.CreateUAV(Output.Texture);
+		PassParameters->RWOutputTexture = GraphBuilder.CreateUAV(KeyerComputeTexture);
 		PassParameters->Output = OutputParameters;
 		PassParameters->SvPositionToInputTextureUV = (
-			FScreenTransform::ChangeTextureBasisFromTo(FScreenPassTextureViewport(Output), FScreenTransform::ETextureBasis::TexelPosition, FScreenTransform::ETextureBasis::ViewportUV) *
+			FScreenTransform::ChangeTextureBasisFromTo(KeyerViewport, FScreenTransform::ETextureBasis::TexelPosition, FScreenTransform::ETextureBasis::ViewportUV) *
 			FScreenTransform::ChangeTextureBasisFromTo(FScreenPassTextureViewport(Input), FScreenTransform::ETextureBasis::ViewportUV, FScreenTransform::ETextureBasis::TextureUV));
 		PassParameters->KeyColor = FVector3f(KeyColor);
 		PassParameters->Params0 = Params0;
@@ -175,7 +189,6 @@ namespace UE::Composite::Private
 		PermutationVector.Set<FCompositePassColorKeyerShader::FUseCleanPlate>(CleanPlateResource != nullptr);
 		PermutationVector.Set<FCompositePassColorKeyerShader::FDespillOnly>(bDespillOnly);
 
-		const FIntPoint OutputSize = Output.ViewRect.Size();
 		TShaderMapRef<FCompositePassColorKeyerShader> ComputeShader(GlobalShaderMap, PermutationVector);
 		FComputeShaderUtils::AddPass(
 			GraphBuilder,
@@ -191,6 +204,9 @@ namespace UE::Composite::Private
 			GraphBuilder.UseExternalAccessMode(CleanPlateResource, ERHIAccess::SRVMask);
 		}
 
+		/** Copy the CS result to the actual output render target, bridging UAV -> CopyDest -> SRVGraphics. */
+		AddCopyTexturePass(GraphBuilder, KeyerComputeTexture, Output.Texture);
+
 		return UE::CompositeCore::FPassTexture{ MoveTemp(Output), Metadata };
 	}
 }
diff --git a/Engine/Plugins/Compositing/Composite/Source/Composite/Private/Passes/CompositePassDilation.cpp b/Engine/Plugins/Compositing/Composite/Source/Composite/Private/Passes/CompositePassDilation.cpp
--- a/Engine/Plugins/Compositing/Composite/Source/Composite/Private/Passes/CompositePassDilation.cpp
+++ b/Engine/Plugins/Compositing/Composite/Source/Composite/Private/Passes/CompositePassDilation.cpp
@@ -16,10 +16,10 @@ class FCompositePassDilationShader : public FGlobalShader
 	SHADER_USE_PARAMETER_STRUCT(FCompositePassDilationShader, FGlobalShader);
 
 	// DILATION_SIZE drives the LDS tile halo and the neighbor kernel:
-	//   1 → 8-neighbor (3×3 minus center), 1 pixel per dispatch
-	//   2 → 24-neighbor (5×5 minus center), 2 pixels per dispatch
+	//   1 -> 8-neighbor (3x3 minus center), 1 pixel per dispatch
+	//   2 -> 24-neighbor (5x5 minus center), 2 pixels per dispatch
 	// The C++ dispatch loop uses DILATION_SIZE=2 greedily, then DILATION_SIZE=1 for any
-	// remaining pixel — same reach as N×DILATION_SIZE=1 but up to half the dispatch count.
+	// remaining pixel - same reach as N*DILATION_SIZE=1 but up to half the dispatch count.
 	class FDilationSize : SHADER_PERMUTATION_RANGE_INT("DILATION_SIZE", 1, 2);
 	class FErode : SHADER_PERMUTATION_BOOL("MORPHOLOGY_ERODE");
 	using FPermutationDomain = TShaderPermutationDomain<FDilationSize, FErode>;
@@ -77,19 +77,30 @@ FPassTexture FCompositePassDilationProxy::Add(FRDGBuilder& GraphBuilder, const F
 		Output = CreateOutputRenderTarget(GraphBuilder, InView, InputTexture.ViewRect, InputTexture.Texture->Desc, TEXT("CompositeDilationOutput"));
 	}
 
-	// Input and Output viewports must be the same size: the dilation kernel operates at a fixed
-	// resolution throughout the ping-pong chain, and the final iteration writes directly to Output.
+	/** Input and Output viewports must be the same size: the dilation kernel operates at a fixed
+	 *  resolution throughout the ping-pong chain, and the final iteration writes to FinalComputeTexture
+	 *  (copied to Output afterwards). */
 	const FIntPoint TextureSize = InputTexture.ViewRect.Size();
 	checkf(Output.ViewRect.Size() == TextureSize,
 		TEXT("CompositePassDilation: Input viewport (%dx%d) and Output viewport (%dx%d) must be the same size."),
 		TextureSize.X, TextureSize.Y, Output.ViewRect.Size().X, Output.ViewRect.Size().Y);
 
-	const FRDGTextureDesc PingPongDesc = FRDGTextureDesc::Create2D(TextureSize, Output.Texture->Desc.Format, FClearValueBinding::Black, TexCreate_UAV | TexCreate_ShaderResource | TexCreate_RenderTargetable);
+	const FRDGTextureDesc PingPongDesc = FRDGTextureDesc::Create2D(TextureSize, Output.Texture->Desc.Format, FClearValueBinding::Black, TexCreate_UAV | TexCreate_ShaderResource);
 	FRDGTextureRef PingPong[2] = {
 		GraphBuilder.CreateTexture(PingPongDesc, TEXT("CompositeDilationPingPong0")),
 		GraphBuilder.CreateTexture(PingPongDesc, TEXT("CompositeDilationPingPong1"))
 	};
 
+	/** The final compute step writes here instead of directly to Output.Texture.
+	 *  Output.Texture is RenderTargetable; D3D12 state tracking for UAV writes onto
+	 *  RenderTargetable textures can race with the RHI thread (r.RHIThread.Enable 1),
+	 *  corrupting the subsequent graphics SRV read in the merge pass. Using a
+	 *  UAV-only intermediate (same pattern as GeometryMaskPostProcess_Blur and
+	 *  PostProcessUpscale) gives RDG a clean UAV -> CopySrc -> CopyDest -> SRVGraphics chain
+	 *  and avoids the issue entirely. */
+	const FRDGTextureDesc FinalComputeDesc = FRDGTextureDesc::Create2D(TextureSize, Output.Texture->Desc.Format, FClearValueBinding::Black, TexCreate_UAV | TexCreate_ShaderResource);
+	FRDGTextureRef FinalComputeTexture = GraphBuilder.CreateTexture(FinalComputeDesc, TEXT("CompositeDilationFinalCS"));
+
 	// Greedy dispatch: consume DILATION_SIZE=2 steps (2 pixels each) until fewer than 2 pixels
 	// remain, then finish with DILATION_SIZE=1. Same total pixel reach as Size×DILATION_SIZE=1
 	// but up to half the number of GPU dispatches.
@@ -111,22 +122,21 @@ FPassTexture FCompositePassDilationProxy::Add(FRDGBuilder& GraphBuilder, const F
 			ReadTexture = FScreenPassTexture(PingPong[ReadIndex], FIntRect(FIntPoint::ZeroValue, TextureSize));
 		}
 
-		// Final iteration writes directly to Output; intermediate writes to ping-pong.
+		/** Final iteration writes to the UAV-only intermediate; intermediate iterations write to ping-pong. */
 		const bool bFinalIteration = (Pixel + StepSize >= Size);
 		FRDGTextureRef WriteTexture;
 		FScreenPassTextureViewportParameters OutputViewportParams;
 		if (bFinalIteration)
 		{
-			WriteTexture = Output.Texture;
-			OutputViewportParams = GetScreenPassTextureViewportParameters(FScreenPassTextureViewport(Output));
+			WriteTexture = FinalComputeTexture;
 		}
 		else
 		{
 			const int32 WriteIndex = Iteration % 2;
 			WriteTexture = PingPong[WriteIndex];
-			OutputViewportParams = GetScreenPassTextureViewportParameters(
-				FScreenPassTextureViewport(FScreenPassTexture(WriteTexture, FIntRect(FIntPoint::ZeroValue, TextureSize))));
 		}
+		OutputViewportParams = GetScreenPassTextureViewportParameters(
+			FScreenPassTextureViewport(FScreenPassTexture(WriteTexture, FIntRect(FIntPoint::ZeroValue, TextureSize))));
 
 		FCompositePassDilationShader::FPermutationDomain PermutationVector;
 		PermutationVector.Set<FCompositePassDilationShader::FDilationSize>(StepSize);
@@ -141,18 +151,22 @@ FPassTexture FCompositePassDilationProxy::Add(FRDGBuilder& GraphBuilder, const F
 		PassParameters->bOpacifyOutput = 0;
 		PassParameters->AlphaThreshold = AlphaThreshold;
 
-		// Dispatch over the write viewport so the final iteration always covers the full output.
-		const FIntPoint DispatchSize = bFinalIteration ? Output.ViewRect.Size() : TextureSize;
 		FComputeShaderUtils::AddPass(
 			GraphBuilder,
 			RDG_EVENT_NAME("Composite.Dilation %s (%dx%d) [pix:%d/%d step:%d]", bErode ? TEXT("Erode") : TEXT("Dilate"), TextureSize.X, TextureSize.Y, Pixel, Size, StepSize),
 			ERDGPassFlags::Compute,
 			ComputeShader,
 			PassParameters,
-			FComputeShaderUtils::GetGroupCount(DispatchSize, FCompositePassDilationShader::ThreadGroupSize)
+			FComputeShaderUtils::GetGroupCount(TextureSize, FCompositePassDilationShader::ThreadGroupSize)
 		);
 	}
 
+	/** Copy the final CS result from the UAV-only intermediate into the RT-capable output.
+	 *  This bridges the compute UAV state on the source to copy-dest on the destination,
+	 *  after which the downstream merge pass transitions cleanly to SRVGraphics. This mirrors
+	 *  the engine pattern in PostProcessUpscale.cpp and GeometryMaskPostProcess_Blur.cpp. */
+	AddCopyTexturePass(GraphBuilder, FinalComputeTexture, Output.Texture);
+
 	return UE::CompositeCore::FPassTexture{ MoveTemp(Output), Metadata };
 }
 
