Composure: Small comment tweak to clarify that composite meshes still need to be assigned to a plate layer.

#Synced-CL 52419094

diff --git a/Engine/Plugins/Compositing/Composite/Source/Composite/Public/CompositeDepthMeshActor.h b/Engine/Plugins/Compositing/Composite/Source/Composite/Public/CompositeDepthMeshActor.h
--- a/Engine/Plugins/Compositing/Composite/Source/Composite/Public/CompositeDepthMeshActor.h
+++ b/Engine/Plugins/Compositing/Composite/Source/Composite/Public/CompositeDepthMeshActor.h
@@ -11,7 +11,7 @@
 
 class UCompositeDepthMeshComponent;
 
-/** Convenience compositing depth mesh actor. */
+/** Convenience compositing depth mesh actor. This actor must still be assigned to a plate layer. */
 UCLASS(MinimalAPI)
 class ACompositeDepthMeshActor : public AActor
 {
diff --git a/Engine/Plugins/Compositing/Composite/Source/Composite/Public/CompositeMeshActor.h b/Engine/Plugins/Compositing/Composite/Source/Composite/Public/CompositeMeshActor.h
--- a/Engine/Plugins/Compositing/Composite/Source/Composite/Public/CompositeMeshActor.h
+++ b/Engine/Plugins/Compositing/Composite/Source/Composite/Public/CompositeMeshActor.h
@@ -11,7 +11,7 @@
 
 class UCompositeMeshComponent;
 
-/** Convenience compositing mesh actor. */
+/** Convenience compositing mesh actor. This actor must still be assigned to a plate layer. */
 UCLASS(MinimalAPI)
 class ACompositeMeshActor : public AActor
 {
