Composure: Show layer passes above media passes in plate Layer pass customization, since everything is bottom-up.

Additional fixes:
- Active pill respects group when adding: OnPassAdded now checks the active group filter pill before falling back to the class-based GetDefaultGroupForNewPass, so adding a pass while "Layer Passes" is active correctly targets the layer group.
- Filter not cleared after add: ClearFilters is now skipped when a group pill is active, since the newly added pass is already visible in the current filter context.
- Add menu filtered by active pill: OnFilterNewPassType now validates pass types against the active group pill when no row is selected, hiding incompatible pass types from the menu.

#Synced-CL 52419094

diff --git a/Engine/Plugins/Compositing/Composite/Source/CompositeEditor/Private/UI/CompositePlatePassListOwner.cpp b/Engine/Plugins/Compositing/Composite/Source/CompositeEditor/Private/UI/CompositePlatePassListOwner.cpp
--- a/Engine/Plugins/Compositing/Composite/Source/CompositeEditor/Private/UI/CompositePlatePassListOwner.cpp
+++ b/Engine/Plugins/Compositing/Composite/Source/CompositeEditor/Private/UI/CompositePlatePassListOwner.cpp
@@ -64,8 +64,8 @@ FString FCompositePlatePassListOwner::GetGroupFilterString(int32 InGroupIndex)
 {
 	static FString PassTypeFilterStrings[] =
 	{
-		TEXT("Media Passes"),
-		TEXT("Layer Passes")
+		TEXT("Layer Passes"),
+		TEXT("Media Passes")
 	};
 
 	return PassTypeFilterStrings[InGroupIndex];
@@ -75,8 +75,8 @@ const FSlateBrush* FCompositePlatePassListOwner::GetGroupIcon(int32 InGroupIndex
 {
 	static const FSlateBrush* PassTypeIcons[] =
 	{
-		FCompositeEditorStyle::Get().GetBrush("CompositeEditor.Passes.Media"),
-		FCompositeEditorStyle::Get().GetBrush("CompositeEditor.Passes.Layer")
+		FCompositeEditorStyle::Get().GetBrush("CompositeEditor.Passes.Layer"),
+		FCompositeEditorStyle::Get().GetBrush("CompositeEditor.Passes.Media")
 	};
 		
 	return PassTypeIcons[InGroupIndex];
@@ -86,8 +86,8 @@ FText FCompositePlatePassListOwner::GetGroupDisplayName(int32 InGroupIndex)
 {
 	static FText PassTypeLabels[] =
 	{
-		LOCTEXT("MediaPassTypeLabel", "Media Passes"),
 		LOCTEXT("LayerPassTypeLabel", "Layer Passes"),
+		LOCTEXT("MediaPassTypeLabel", "Media Passes"),
 	};
 		
 	return PassTypeLabels[InGroupIndex];
@@ -97,8 +97,8 @@ ICompositePassListOwner::FGroupFilterConfig FCompositePlatePassListOwner::GetGro
 {
 	static TArray<FGroupFilterConfig> PassTypeFilterConfigs =
 	{
-		{ TEXT("MediaPassesFilter"), LOCTEXT("MediaPassesFilterName", "Media Passes"),  LOCTEXT("MediaPassesFilterToolTip", "Only show media passes") },
-		{ TEXT("LayerPassesFilter"), LOCTEXT("LayerPassesFilterName", "Layer Passes"),  LOCTEXT("LayerPassesFilterToolTip", "Only show layer passes")  }
+		{ TEXT("LayerPassesFilter"), LOCTEXT("LayerPassesFilterName", "Layer Passes"),  LOCTEXT("LayerPassesFilterToolTip", "Only show layer passes")  },
+		{ TEXT("MediaPassesFilter"), LOCTEXT("MediaPassesFilterName", "Media Passes"),  LOCTEXT("MediaPassesFilterToolTip", "Only show media passes") }
 	};
 
 	return PassTypeFilterConfigs[InGroupIndex];
diff --git a/Engine/Plugins/Compositing/Composite/Source/CompositeEditor/Private/UI/CompositePlatePassListOwner.h b/Engine/Plugins/Compositing/Composite/Source/CompositeEditor/Private/UI/CompositePlatePassListOwner.h
--- a/Engine/Plugins/Compositing/Composite/Source/CompositeEditor/Private/UI/CompositePlatePassListOwner.h
+++ b/Engine/Plugins/Compositing/Composite/Source/CompositeEditor/Private/UI/CompositePlatePassListOwner.h
@@ -9,12 +9,11 @@
 class FCompositePlatePassListOwner : public ICompositePassListOwner
 {
 private:
-	/** Pass type which corresponds to which list in the plate layer a pass is stored in */
+	/** Pass type which corresponds to which list in the plate layer a pass is stored in. The enum order determines display order — media passes are last so they appear at the bottom. */
 	enum class EPassType : uint8
 	{
-		Media,
-		//Scene, // TODO: Add Scene when it is supported
 		Layer,
+		Media,
 		MAX
 	};
 
diff --git a/Engine/Plugins/Compositing/Composite/Source/CompositeEditor/Private/UI/SCompositePassTree.cpp b/Engine/Plugins/Compositing/Composite/Source/CompositeEditor/Private/UI/SCompositePassTree.cpp
--- a/Engine/Plugins/Compositing/Composite/Source/CompositeEditor/Private/UI/SCompositePassTree.cpp
+++ b/Engine/Plugins/Compositing/Composite/Source/CompositeEditor/Private/UI/SCompositePassTree.cpp
@@ -245,6 +245,20 @@ public:
 		return TreeTextFilter->PassesFilter(InItem) && PassesAnyFilterBarFilter(InItem);
 	}
 
+	/** Returns the group index of the active group filter pill, or INDEX_NONE if "All Passes" (or no filter) is active */
+	int32 GetActiveGroupFilterIndex() const
+	{
+		for (int32 Index = 0; Index < GroupFilters.Num(); ++Index)
+		{
+			if (GroupFilters[Index].IsValid() && GroupFilters[Index]->IsActive())
+			{
+				return Index;
+			}
+		}
+
+		return INDEX_NONE;
+	}
+
 	void ClearFilters()
 	{
 		// Prevent filter changed callback from firing while changing filter status, will invoke once all filters are reset
@@ -1656,8 +1670,9 @@ void SCompositePassTree::OnPassAdded(const UClass* InPassClass)
 	TArray<FPassTreeItemPtr> SelectedItems = TreeView.IsValid() ? TreeView->GetSelectedItems() : TArray<FPassTreeItemPtr>();
 	if (SelectedItems.IsEmpty())
 	{
-		// If there are no selected items, add the new pass to the default group, or after the last pass if there are no groups
-		const int32 DefaultGroupIndex = PassListOwner->GetDefaultGroupForNewPass(InPassClass);
+		// If a group filter pill is active, honour it; otherwise fall back to the class-based default group.
+		const int32 ActiveGroupIndex = Toolbar.IsValid() ? Toolbar->GetActiveGroupFilterIndex() : INDEX_NONE;
+		const int32 DefaultGroupIndex = (ActiveGroupIndex != INDEX_NONE) ? ActiveGroupIndex : PassListOwner->GetDefaultGroupForNewPass(InPassClass);
 		if (DefaultGroupIndex != INDEX_NONE)
 		{
 			if (PassTreeItems.IsValidIndex(DefaultGroupIndex))
@@ -1708,7 +1723,9 @@ void SCompositePassTree::OnPassAdded(const UClass* InPassClass)
 		Composite::Analytics::RecordPassAdded(*NewPass);
 	}
 
-	if (Toolbar.IsValid())
+	// Clear filters so the new pass is visible — but only when no group filter is active.
+	// If a group filter pill is active, the pass was just added to that group and is already visible.
+	if (Toolbar.IsValid() && Toolbar->GetActiveGroupFilterIndex() == INDEX_NONE)
 	{
 		Toolbar->ClearFilters();
 	}
@@ -1757,6 +1774,18 @@ bool SCompositePassTree::OnFilterNewPassType(const UClass* InPassType) const
 	}
 
 	TArray<FPassTreeItemPtr> SelectedItems = TreeView.IsValid() ? TreeView->GetSelectedItems() : TArray<FPassTreeItemPtr>();
+	if (SelectedItems.IsEmpty())
+	{
+		// No row selected — if a group filter pill is active, validate against that group.
+		const int32 ActiveGroupIndex = Toolbar.IsValid() ? Toolbar->GetActiveGroupFilterIndex() : INDEX_NONE;
+		if (ActiveGroupIndex != INDEX_NONE)
+		{
+			return !PassListOwner->CanAddPass(InPassType, ActiveGroupIndex, INDEX_NONE);
+		}
+
+		return false;
+	}
+
 	for (const FPassTreeItemPtr& Item : SelectedItems)
 	{
 		if (!PassListOwner->CanAddPass(InPassType, Item->GroupIndex, Item->PassIndex))
@@ -1765,7 +1794,7 @@ bool SCompositePassTree::OnFilterNewPassType(const UClass* InPassType) const
 			return true;
 		}
 	}
-	
+
 	return false;
 }
 
