Composure: Added ability to include or exclude primitives as composite mesh primitives in the Plate layer's composite mesh actor list

#Synced-CL 52432111

diff --git a/Engine/Plugins/Compositing/Composite/Source/Composite/Private/Layers/CompositeLayerPlate.cpp b/Engine/Plugins/Compositing/Composite/Source/Composite/Private/Layers/CompositeLayerPlate.cpp
--- a/Engine/Plugins/Compositing/Composite/Source/Composite/Private/Layers/CompositeLayerPlate.cpp
+++ b/Engine/Plugins/Compositing/Composite/Source/Composite/Private/Layers/CompositeLayerPlate.cpp
@@ -105,7 +105,38 @@ namespace UE
 					}
 				}
 			}
-
+			
+			FCompositeMeshPrimitiveReference* FindPrimitiveReferences(TArray<FCompositeMeshPrimitiveReference>& PrimitiveReferences, const TSoftObjectPtr<AActor>& InActor)
+			{
+				return PrimitiveReferences.FindByPredicate([&InActor](const FCompositeMeshPrimitiveReference& Item)
+				{
+					return Item.CompositeMeshActor == InActor;
+				});
+			}
+			
+			const FCompositeMeshPrimitiveReference* FindPrimitiveReferences(const TArray<FCompositeMeshPrimitiveReference>& PrimitiveReferences, const TSoftObjectPtr<AActor>& InActor)
+			{
+				return PrimitiveReferences.FindByPredicate([&InActor](const FCompositeMeshPrimitiveReference& Item)
+				{
+					return Item.CompositeMeshActor == InActor;
+				});
+			}
+			
+			void CleanUpPrimitiveReference(FCompositeMeshPrimitiveReference& PrimitiveReference)
+			{
+				// Don't try to clean up any soft pointers that aren't explicitly loaded
+				if (!PrimitiveReference.CompositeMeshActor.IsValid())
+				{
+					return;
+				}
+				
+				TInlineComponentArray<UPrimitiveComponent*> PrimComponents(PrimitiveReference.CompositeMeshActor.Get());
+				PrimitiveReference.CompositeMeshPrimitives.RemoveAll([&PrimComponents](const TSoftObjectPtr<UPrimitiveComponent>& Primitive)
+				{
+					return !Primitive.IsValid() || !PrimComponents.Contains(Primitive.Get());
+				});
+			}
+			
 			/** Return the texture size, depending on the texture type. */
 			static FIntPoint GetTextureSize(const UTexture* InTexture)
 			{
@@ -299,6 +330,15 @@ void UCompositeLayerPlate::SetCompositeMeshes(TArray<TSoftObjectPtr<AActor>> InC
 
 	CompositeMeshes = MoveTemp(InCompositeMeshes);
 
+	CompositeMeshPrimitives.Empty();
+	for (const TSoftObjectPtr<AActor>& Actor : CompositeMeshes)
+	{
+		FCompositeMeshPrimitiveReference NewPrimitiveReference;
+		NewPrimitiveReference.CompositeMeshActor = Actor;
+		
+		CompositeMeshPrimitives.Add(NewPrimitiveReference);
+	}
+
 	UpdateCompositeMeshes();
 }
 
@@ -409,6 +449,11 @@ void UCompositeLayerPlate::PostEditChangeProperty(FPropertyChangedEvent& Propert
 
 		if (bChangeWasRemoval)
 		{
+			CompositeMeshPrimitives.RemoveAll([this](const FCompositeMeshPrimitiveReference& Item)
+			{
+				return !CompositeMeshes.Contains(Item.CompositeMeshActor);
+			});
+
 			UE::Composite::Private::HandleRemovedCompositeMeshes(PreEditCompositeMeshes, CompositeMeshes, bControlMeshVisibility);
 		}
 		else
@@ -430,6 +475,17 @@ void UCompositeLayerPlate::PostEditChangeProperty(FPropertyChangedEvent& Propert
 
 						NewCompositeMesh = nullptr;
 					}
+					else
+					{
+						FCompositeMeshPrimitiveReference* PrimitiveReference = UE::Composite::Private::FindPrimitiveReferences(CompositeMeshPrimitives, NewCompositeMesh);
+						if (!PrimitiveReference)
+						{
+							FCompositeMeshPrimitiveReference NewPrimitiveReference;
+							NewPrimitiveReference.CompositeMeshActor = NewCompositeMesh;
+							
+							CompositeMeshPrimitives.Add(NewPrimitiveReference);
+						}
+					}
 				}
 			}
 		}
@@ -701,7 +757,32 @@ TArray<UPrimitiveComponent*> UCompositeLayerPlate::GetPrimitives() const
 			continue;
 		}
 
-		TInlineComponentArray<UPrimitiveComponent*> PrimComponents(Actor.Get());
+		OutPrimitiveComponents.Append(GetPrimitivesForActor(Actor));
+	}
+
+	return OutPrimitiveComponents;
+}
+
+TArray<UPrimitiveComponent*> UCompositeLayerPlate::GetPrimitivesForActor(TSoftObjectPtr<AActor> InActor) const
+{
+	if (!InActor.IsValid())
+	{
+		return TArray<UPrimitiveComponent*>();
+	}
+	
+	if (!CompositeMeshes.Contains(InActor))
+	{
+		return TArray<UPrimitiveComponent*>();
+	}
+	
+	const UCompositeCorePluginSettings* Settings = GetDefault<UCompositeCorePluginSettings>();
+	check(Settings);
+	
+	TArray<UPrimitiveComponent*> OutPrimitiveComponents = GetIncludedPrimitivesForActor(InActor);
+	if (OutPrimitiveComponents.IsEmpty())
+	{
+		// If no primitives were specifically included for the actor, return all valid primitives for the actor
+		TInlineComponentArray<UPrimitiveComponent*> PrimComponents(InActor.Get());
 
 		for (UPrimitiveComponent* PrimitiveComponent : PrimComponents)
 		{
@@ -715,6 +796,116 @@ TArray<UPrimitiveComponent*> UCompositeLayerPlate::GetPrimitives() const
 	return OutPrimitiveComponents;
 }
 
+TArray<UPrimitiveComponent*> UCompositeLayerPlate::GetIncludedPrimitivesForActor(TSoftObjectPtr<AActor> InActor) const
+{
+	if (!InActor.IsValid())
+	{
+		return TArray<UPrimitiveComponent*>();
+	}
+	
+	if (!CompositeMeshes.Contains(InActor))
+	{
+		return TArray<UPrimitiveComponent*>();
+	}
+	
+	TArray<UPrimitiveComponent*> OutPrimitiveComponents;
+	if (const FCompositeMeshPrimitiveReference* PrimitiveReference = UE::Composite::Private::FindPrimitiveReferences(CompositeMeshPrimitives, InActor))
+	{
+		for (const TSoftObjectPtr<UPrimitiveComponent>& PrimitiveComponent : PrimitiveReference->CompositeMeshPrimitives)
+		{
+			if (!PrimitiveComponent.IsValid())
+			{
+				continue;
+			}
+		
+			OutPrimitiveComponents.Add(PrimitiveComponent.Get());
+		}
+	}
+	
+	return OutPrimitiveComponents;
+}
+
+void UCompositeLayerPlate::SetPrimitiveIncluded(UPrimitiveComponent* InPrimitive, bool bInIncluded)
+{
+	AActor* PrimitiveOwner = InPrimitive->GetOwner();
+	if (!CompositeMeshes.Contains(PrimitiveOwner))
+	{
+		return;
+	}
+	
+	FCompositeMeshPrimitiveReference* PrimitiveReference = UE::Composite::Private::FindPrimitiveReferences(CompositeMeshPrimitives, PrimitiveOwner);
+	if (!PrimitiveReference)
+	{
+		FCompositeMeshPrimitiveReference NewPrimitiveReference;
+		NewPrimitiveReference.CompositeMeshActor = PrimitiveOwner;
+		
+		const int32 Index = CompositeMeshPrimitives.Add(NewPrimitiveReference);
+		PrimitiveReference = &CompositeMeshPrimitives[Index];
+	}
+	
+	const UCompositeCorePluginSettings* Settings = GetDefault<UCompositeCorePluginSettings>();
+	check(Settings);
+	if (!Settings->IsAllowedPrimitiveClass(InPrimitive))
+	{
+		return;
+	}
+	
+	// If there are no included primitives in the list before adding the new primitive, we must unregister all the actor's valid primitives
+	if (bInIncluded && PrimitiveReference->CompositeMeshPrimitives.IsEmpty())
+	{
+		TArray<UPrimitiveComponent*> AllPrimitives = GetPrimitivesForActor(PrimitiveOwner);
+		UnregisterCompositePrimitives(AllPrimitives);
+		for (UPrimitiveComponent* Primitive : AllPrimitives)
+		{
+			PrimitiveReference->CompositeMeshPrimitives.Remove(Primitive);
+
+			Primitive->RemoveUserDataOfClass(UCompositeAssetUserData::StaticClass());
+
+			if (bControlMeshVisibility)
+			{
+				UE::Composite::Private::SetCompositeMeshVisibility(*Primitive, true);
+			}
+		}
+	}
+	
+	if (bInIncluded && !PrimitiveReference->CompositeMeshPrimitives.Contains(InPrimitive))
+	{
+		PrimitiveReference->CompositeMeshPrimitives.Add(InPrimitive);
+		
+		UCompositeAssetUserData* AssetUserData = NewObject<UCompositeAssetUserData>(InPrimitive);
+		AssetUserData->OnPostEditChangeOwner.BindUObject(this, &UCompositeLayerPlate::UpdatePrimitiveComponent);
+		InPrimitive->AddAssetUserData(AssetUserData);
+		TryRegisterCompositePrimitives({ InPrimitive });
+		UpdatePrimitiveComponent(*InPrimitive);
+	}
+	else if (!bInIncluded && PrimitiveReference->CompositeMeshPrimitives.Contains(InPrimitive))
+	{
+		UnregisterCompositePrimitives({ InPrimitive });
+		PrimitiveReference->CompositeMeshPrimitives.Remove(InPrimitive);
+
+		InPrimitive->RemoveUserDataOfClass(UCompositeAssetUserData::StaticClass());
+
+		if (bControlMeshVisibility)
+		{
+			UE::Composite::Private::SetCompositeMeshVisibility(*InPrimitive, true);
+		}
+	}
+	
+	// If there are no primitives specifically included, all primitives in the actor are used as composite mesh primitives, so configure all of them
+	if (!bInIncluded && PrimitiveReference->CompositeMeshPrimitives.IsEmpty())
+	{
+		TArray<UPrimitiveComponent*> AllPrimitives = GetPrimitivesForActor(PrimitiveOwner);
+		TryRegisterCompositePrimitives(AllPrimitives);
+		for (UPrimitiveComponent* Primitive : AllPrimitives)
+		{
+			UCompositeAssetUserData* AssetUserData = NewObject<UCompositeAssetUserData>(Primitive);
+			AssetUserData->OnPostEditChangeOwner.BindUObject(this, &UCompositeLayerPlate::UpdatePrimitiveComponent);
+			Primitive->AddAssetUserData(AssetUserData);
+			UpdatePrimitiveComponent(*Primitive);
+		}
+	}
+}
+
 int32 UCompositeLayerPlate::GetValidPrimitivesNum() const
 {
 	const UCompositeCorePluginSettings* Settings = GetDefault<UCompositeCorePluginSettings>();
@@ -728,14 +919,28 @@ int32 UCompositeLayerPlate::GetValidPrimitivesNum() const
 		{
 			continue;
 		}
-
-		TInlineComponentArray<UPrimitiveComponent*> PrimComponents(Actor.Get());
-
-		for (UPrimitiveComponent* PrimitiveComponent : PrimComponents)
+		
+		// If there is an entry in CompositeMeshPrimitives for the actor, use only the primitives in that entry; otherwise, use all primitives from the actor
+		if (const FCompositeMeshPrimitiveReference* PrimitiveReference = UE::Composite::Private::FindPrimitiveReferences(CompositeMeshPrimitives, Actor))
 		{
-			if (Settings->IsAllowedPrimitiveClass(PrimitiveComponent))
+			for (const TSoftObjectPtr<UPrimitiveComponent>& PrimitiveComponent : PrimitiveReference->CompositeMeshPrimitives)
 			{
-				++Count;
+				if (Settings->IsAllowedPrimitiveClass(PrimitiveComponent.Get()))
+				{
+					++Count;
+				}
+			}
+		}
+		else
+		{
+			TInlineComponentArray<UPrimitiveComponent*> PrimComponents(Actor.Get());
+
+			for (UPrimitiveComponent* PrimitiveComponent : PrimComponents)
+			{
+				if (Settings->IsAllowedPrimitiveClass(PrimitiveComponent))
+				{
+					++Count;
+				}
 			}
 		}
 	}
@@ -813,10 +1018,24 @@ void UCompositeLayerPlate::OnPreSaveWorld(UWorld* InWorld, FObjectPreSaveContext
 	// We need to remove our asset user data before saving, as we do not need to save it out
 	// and only use it to know when the static mesh component changes.
 	UE::Composite::Private::RemoveCompositeAllAssetUserData(CompositeMeshes);
+
+	// Copy the existing CompositeMeshPrimitives to preserve any non-valid soft pointers. We do this to
+	// avoid having to save invalid soft pointers to primitive components while keeping them in memory so
+	// that actions like undoing the deletion of a primitive on a composite mesh actor can be preserved
+	PreSaveCompositeMeshPrimitives = CompositeMeshPrimitives;
+
+	// Clean up any invalid soft pointers to primitives in the list of composite mesh primitives
+	for (FCompositeMeshPrimitiveReference& PrimitiveReference : CompositeMeshPrimitives)
+	{
+		UE::Composite::Private::CleanUpPrimitiveReference(PrimitiveReference);
+	}
 }
 
 void UCompositeLayerPlate::OnPostSaveWorld(UWorld* InWorld, FObjectPostSaveContext ObjectSaveContext)
 {
+	CompositeMeshPrimitives = PreSaveCompositeMeshPrimitives;
+	PreSaveCompositeMeshPrimitives.Empty();
+
 	constexpr bool bRegisterOnly = true;
 	UpdateCompositeMeshes(bRegisterOnly);
 }
@@ -926,15 +1145,9 @@ void UCompositeLayerPlate::UpdateCompositeMeshes(bool bRegisterOnly) const
 			continue;
 		}
 
-		for (UActorComponent* Component : MeshActor->GetComponents())
+		TArray<UPrimitiveComponent*> Primitives = GetPrimitivesForActor(MeshActor);
+		for (UPrimitiveComponent* PrimitiveComponent : Primitives)
 		{
-			UPrimitiveComponent* PrimitiveComponent = Cast<UPrimitiveComponent>(Component);
-
-			if (!IsValid(PrimitiveComponent))
-			{
-				continue;
-			}
-
 			// Track meshes so that we can update their material upon changes
 			if (!PrimitiveComponent->HasAssetUserDataOfClass(UCompositeAssetUserData::StaticClass()))
 			{
@@ -1033,19 +1246,17 @@ void UCompositeLayerPlate::PropagateStateChange(ECompositeStateChangeType Change
 				continue;
 			}
 
-			for (UActorComponent* Component : MeshActor->GetComponents())
+			TArray<UPrimitiveComponent*> PrimitivesForActor = GetPrimitivesForActor(MeshActor);
+			for (UPrimitiveComponent* PrimitiveComponent : PrimitivesForActor)
 			{
-				if (UPrimitiveComponent* PrimitiveComponent = Cast<UPrimitiveComponent>(Component))
+				for (int32 ElementIndex = 0; ElementIndex < PrimitiveComponent->GetNumMaterials(); ++ElementIndex)
 				{
-					for (int32 ElementIndex = 0; ElementIndex < PrimitiveComponent->GetNumMaterials(); ++ElementIndex)
+					if (UMaterialInterface* MaterialInterface = PrimitiveComponent->GetMaterial(ElementIndex))
 					{
-						if (UMaterialInterface* MaterialInterface = PrimitiveComponent->GetMaterial(ElementIndex))
+						if (IsUsingAfterDOFTranslucency(MaterialInterface))
 						{
-							if (IsUsingAfterDOFTranslucency(MaterialInterface))
-							{
-								constexpr bool bOverride = true;
-								OverrideTranslucencyConsoleVariables(bOverride);
-							}
+							constexpr bool bOverride = true;
+							OverrideTranslucencyConsoleVariables(bOverride);
 						}
 					}
 				}
diff --git a/Engine/Plugins/Compositing/Composite/Source/Composite/Public/Layers/CompositeLayerPlate.h b/Engine/Plugins/Compositing/Composite/Source/Composite/Public/Layers/CompositeLayerPlate.h
--- a/Engine/Plugins/Compositing/Composite/Source/Composite/Public/Layers/CompositeLayerPlate.h
+++ b/Engine/Plugins/Compositing/Composite/Source/Composite/Public/Layers/CompositeLayerPlate.h
@@ -24,6 +24,20 @@ enum class ECompositePlateMode : uint8
 	CompositeMesh = 1 UMETA(ToolTip = "Sample the composite meshes rendered in a built-in custom render pass (default). Automatic fallback to Texture."),
 };
 
+USTRUCT()
+struct FCompositeMeshPrimitiveReference
+{
+	GENERATED_BODY()
+	
+	/** The actor whose primitives will be used for the composite meshes */
+	UPROPERTY()
+	TSoftObjectPtr<AActor> CompositeMeshActor;
+	
+	/** The primitives to use for the composite meshes */
+	UPROPERTY()
+	TArray<TSoftObjectPtr<UPrimitiveComponent>> CompositeMeshPrimitives;
+};
+
 /**
 * Convenience layer for media texture integration, with additional processing options.
 * 
@@ -75,6 +89,16 @@ public:
 	/** Get array of all composite mesh primitives. */
 	UE_API TArray<UPrimitiveComponent*> GetPrimitives() const;
 
+	/** Get list of composite mesh primitives for a specified composite mesh actor */
+	UE_API TArray<UPrimitiveComponent*> GetPrimitivesForActor(TSoftObjectPtr<AActor> InActor) const;
+
+	/** Gets the list of primitives that are specifically included for the composite mesh actor. Differs from GetPrimitivesForActor in that
+	 * this returns an empty list if no primitives are specifically included, whereas GetPrimitivesForActor falls back to all primitives */
+	UE_API TArray<UPrimitiveComponent*> GetIncludedPrimitivesForActor(TSoftObjectPtr<AActor> InActor) const;
+	
+	/** Sets whether a composite mesh actor's primitive should be included as a composite mesh primitive or not */
+	UE_API void SetPrimitiveIncluded(UPrimitiveComponent* InPrimitive, bool bInIncluded);
+	
 public:
 	//~ Begin UObject Interface
 	UE_API virtual void PostLoad() override;
@@ -134,6 +158,10 @@ private:
 	UPROPERTY(EditAnywhere, BlueprintGetter = GetCompositeMeshes, BlueprintSetter = SetCompositeMeshes, Category = "Composite", meta = (AllowPrivateAccess = true))
 	TArray<TSoftObjectPtr<AActor>> CompositeMeshes;
 
+	/** A list of actual primitives from the composite mesh actors that will be marked as primitive alpha holdout, and rendered separately with the projected (media) texture  */	
+	UPROPERTY()
+	TArray<FCompositeMeshPrimitiveReference> CompositeMeshPrimitives;
+
 public:
 	/** Media texture for compositing layer. */
 	UPROPERTY(EditAnywhere, BlueprintReadWrite, Category = "Composite", meta = (AllowedClasses="/Script/Engine.Texture2D, /Script/Engine.Texture2DDynamic, /Script/MediaAssets.MediaTexture"))
@@ -262,6 +290,9 @@ private:
 	/** Pre-edit list of registered composite meshes. */
 	TArray<TSoftObjectPtr<AActor>> PreEditCompositeMeshes;
 
+	/** Transient copy of CompositeMeshPrimitives that keeps non-valid soft pointers during a save as they will be removed from CompositeMeshPrimitives during saving */
+	TArray<FCompositeMeshPrimitiveReference> PreSaveCompositeMeshPrimitives;
+	
 	/**
 	 * Called before a level saves.
 	 */
diff --git a/Engine/Plugins/Compositing/Composite/Source/CompositeEditor/Private/Customizations/CompositeLayerPlateCustomization.cpp b/Engine/Plugins/Compositing/Composite/Source/CompositeEditor/Private/Customizations/CompositeLayerPlateCustomization.cpp
--- a/Engine/Plugins/Compositing/Composite/Source/CompositeEditor/Private/Customizations/CompositeLayerPlateCustomization.cpp
+++ b/Engine/Plugins/Compositing/Composite/Source/CompositeEditor/Private/Customizations/CompositeLayerPlateCustomization.cpp
@@ -3,6 +3,7 @@
 #include "CompositeLayerPlateCustomization.h"
 
 #include "CompositeActor.h"
+#include "CompositeCoreSettings.h"
 #include "CompositeEditorStyle.h"
 #include "CompositeMeshActor.h"
 #include "DetailCategoryBuilder.h"
@@ -33,12 +34,15 @@
 #include "Widgets/Images/SImage.h"
 #include "Widgets/Input/SComboButton.h"
 #include "Widgets/Text/STextBlock.h"
+#include "Widgets/Views/SHeaderRow.h"
 
 #define LOCTEXT_NAMESPACE "FCompositeLayerPlateCustomization"
 
 namespace CompositeLayerPlateCustomization
 {
 	DECLARE_DELEGATE_TwoParams(FOnCustomizePropertyRow, const TSharedPtr<IPropertyHandle>&, IDetailPropertyRow&);
+	
+	static FName CompositeMeshListColumn_Components = "Components";
 }
 
 TSharedRef<IDetailCustomization> FCompositeLayerPlateCustomization::MakeInstance()
@@ -111,9 +115,19 @@ void FCompositeLayerPlateCustomization::CustomizeDetails(IDetailLayoutBuilder& D
 				constexpr bool bExcludeCompositeMeshActors = false;
 				return SCompositeActorPickerTable::MakeDefaultSceneOutlinerFilters(bExcludeCompositeMeshActors);
 			})
+			.OnGenerateHeaderColumns_Lambda([](TSharedRef<SHeaderRow> InHeaderRow)
+			{
+				InHeaderRow->AddColumn(
+				SHeaderRow::Column(CompositeLayerPlateCustomization::CompositeMeshListColumn_Components)
+				.DefaultLabel(LOCTEXT("ComponentsColumnLabel", "Components"))
+				.FillWidth(0.4));
+			})
+			.OnGenerateColumnWidget(this, &FCompositeLayerPlateCustomization::GetCompositeMeshColumnWidget)
 			.OnExtendAddMenu(this, &FCompositeLayerPlateCustomization::OnExtendCompositeMeshAddMenu)
+			.OnExtendContextMenu(this, &FCompositeLayerPlateCustomization::OnExtendCompositeMeshContextMenu)
 			.OnLayoutSizeChanged(this, &FCompositeLayerPlateCustomization::OnLayoutSizeChanged)
 			.ShowApplyMaterialSection(true)
+			.OnApplyMaterialToActor(this, &FCompositeLayerPlateCustomization::ApplyMaterialToActor)
 		];
 		
 		IDetailGroup& PassesGroup = PlateCategory.AddGroup("Passes", LOCTEXT("PassesGroupName", "Passes"), false, true);
@@ -306,6 +320,50 @@ void FCompositeLayerPlateCustomization::OpenMediaProfile()
 	GEditor->GetEditorSubsystem<UAssetEditorSubsystem>()->OpenEditorForAsset(ActiveMediaProfile);
 }
 
+TSharedRef<SWidget> FCompositeLayerPlateCustomization::GetCompositeMeshColumnWidget(const TSoftObjectPtr<AActor>& InActor, const FName& InColumnName)
+{
+	if (InColumnName == CompositeLayerPlateCustomization::CompositeMeshListColumn_Components)
+	{
+		return SNew(STextBlock).Text(this, &FCompositeLayerPlateCustomization::GetCompositeMeshComponentsLabel, InActor);
+	}
+	
+	return SNullWidget::NullWidget;
+}
+
+FText FCompositeLayerPlateCustomization::GetCompositeMeshComponentsLabel(TSoftObjectPtr<AActor> InActor) const
+{
+	if (!InActor.IsValid())
+	{
+		return FText::GetEmpty();
+	}
+	
+	TSharedPtr<IDetailLayoutBuilder> DetailBuilder = CachedDetailBuilder.Pin();
+	if (!DetailBuilder.IsValid())
+	{
+		return FText::GetEmpty();
+	}
+	
+	TArray<TWeakObjectPtr<UCompositeLayerPlate>> Plates = DetailBuilder->GetObjectsOfTypeBeingCustomized<UCompositeLayerPlate>();
+	if (Plates.Num() != 1 || !Plates[0].IsValid())
+	{
+		return FText::GetEmpty();
+	}
+	
+	UCompositeLayerPlate* Plate = Plates[0].Get();
+	TArray<UPrimitiveComponent*> CompositeMeshPrimitives = Plate->GetIncludedPrimitivesForActor(InActor);
+	
+	if (CompositeMeshPrimitives.Num() > 1)
+	{
+		return FText::Format(LOCTEXT("MultipleComponentsLabelFormat", "{0} Components"), CompositeMeshPrimitives.Num());
+	}
+	if (CompositeMeshPrimitives.Num() == 1)
+	{
+		return FText::FromString(CompositeMeshPrimitives[0]->GetName());
+	}
+	
+	return FText::GetEmpty(); 
+}
+
 void FCompositeLayerPlateCustomization::OnExtendCompositeMeshAddMenu(FMenuBuilder& MenuBuilder)
 {
 	MenuBuilder.AddMenuEntry(
@@ -318,6 +376,229 @@ void FCompositeLayerPlateCustomization::OnExtendCompositeMeshAddMenu(FMenuBuilde
 	);
 }
 
+void FCompositeLayerPlateCustomization::OnExtendCompositeMeshContextMenu(FMenuBuilder& MenuBuilder, TArray<TSoftObjectPtr<AActor>>& SelectedActors)
+{
+	if (SelectedActors.Num() != 1)
+	{
+		return;
+	}
+	
+	TSoftObjectPtr<AActor> Actor = SelectedActors[0];
+	if (!Actor.IsValid())
+	{
+		return;
+	}
+
+	MenuBuilder.BeginSection("ComponentSection", LOCTEXT("ComponentSectionLabel", "Component"));
+	{
+		FMenuEntryParams SelectedActorEntry;
+		SelectedActorEntry.LabelOverride = LOCTEXT("SelectedActorEntryLabel", "Selected Actor");
+		SelectedActorEntry.ToolTipOverride = LOCTEXT("SelectedActorEntryToolTip", "Select all components of this actor");
+		SelectedActorEntry.IconOverride = FSlateIconFinder::FindIconForClass(Actor->GetClass());
+		SelectedActorEntry.UserInterfaceActionType = EUserInterfaceActionType::ToggleButton;
+		SelectedActorEntry.bShouldCloseWindowAfterMenuSelection = false;
+		SelectedActorEntry.DirectActions = FUIAction(
+			FExecuteAction::CreateSP(this, &FCompositeLayerPlateCustomization::OnCompositeMeshActorToggled, Actor),
+			FCanExecuteAction::CreateSP(this, &FCompositeLayerPlateCustomization::CanToggleCompositeMeshActor, Actor),
+			FIsActionChecked::CreateSP(this, &FCompositeLayerPlateCustomization::IsCompositeMeshActorChecked, Actor));
+		
+		MenuBuilder.AddMenuEntry(SelectedActorEntry);
+		
+		MenuBuilder.AddSeparator();
+		
+		const UCompositeCorePluginSettings* Settings = GetDefault<UCompositeCorePluginSettings>();
+		TInlineComponentArray<UPrimitiveComponent*> PrimComponents(Actor.Get());
+
+		for (UPrimitiveComponent* PrimitiveComponent : PrimComponents)
+		{
+			if (Settings && Settings->IsAllowedPrimitiveClass(PrimitiveComponent))
+			{
+				FMenuEntryParams ComponentEntry;
+				ComponentEntry.LabelOverride = FText::FromString(PrimitiveComponent->GetName());
+				ComponentEntry.IconOverride = FSlateIconFinder::FindIconForClass(PrimitiveComponent->GetClass());
+				ComponentEntry.UserInterfaceActionType = EUserInterfaceActionType::ToggleButton;
+				ComponentEntry.bShouldCloseWindowAfterMenuSelection = false;
+				ComponentEntry.DirectActions = FUIAction(
+					FExecuteAction::CreateSP(this, &FCompositeLayerPlateCustomization::OnCompositeMeshPrimitiveToggled, Actor, PrimitiveComponent),
+					FCanExecuteAction(),
+					FIsActionChecked::CreateSP(this, &FCompositeLayerPlateCustomization::IsCompositeMeshPrimitiveChecked, Actor, PrimitiveComponent));
+		
+				MenuBuilder.AddMenuEntry(ComponentEntry);
+			}
+		}
+	}
+	MenuBuilder.EndSection();
+}
+
+void FCompositeLayerPlateCustomization::OnCompositeMeshActorToggled(TSoftObjectPtr<AActor> InActor)
+{
+	if (!InActor.IsValid())
+	{
+		return;
+	}
+	
+	TSharedPtr<IDetailLayoutBuilder> DetailBuilder = CachedDetailBuilder.Pin();
+	if (!DetailBuilder.IsValid())
+	{
+		return;
+	}
+	
+	TArray<TWeakObjectPtr<UCompositeLayerPlate>> Plates = DetailBuilder->GetObjectsOfTypeBeingCustomized<UCompositeLayerPlate>();
+	if (Plates.Num() != 1 || !Plates[0].IsValid())
+	{
+		return;
+	}
+	
+	UCompositeLayerPlate* Plate = Plates[0].Get();
+	TArray<UPrimitiveComponent*> CompositeMeshPrimitives = Plate->GetIncludedPrimitivesForActor(InActor);
+
+	FScopedTransaction Transaction(LOCTEXT("ToggleCompositeMeshPrimitives", "Toggle Composite Mesh Actor"));
+	Plate->Modify();
+	
+	// Toggling the whole actor means removing all specifically included primitives so that all valid primitives in the actor are used
+	for (UPrimitiveComponent* PrimitiveComponent : CompositeMeshPrimitives)
+	{
+		Plate->SetPrimitiveIncluded(PrimitiveComponent, false);
+	}
+}
+
+bool FCompositeLayerPlateCustomization::CanToggleCompositeMeshActor(TSoftObjectPtr<AActor> InActor) const
+{
+	if (!InActor.IsValid())
+	{
+		return false;
+	}
+	
+	TSharedPtr<IDetailLayoutBuilder> DetailBuilder = CachedDetailBuilder.Pin();
+	if (!DetailBuilder.IsValid())
+	{
+		return false;
+	}
+	
+	TArray<TWeakObjectPtr<UCompositeLayerPlate>> Plates = DetailBuilder->GetObjectsOfTypeBeingCustomized<UCompositeLayerPlate>();
+	if (Plates.Num() != 1 || !Plates[0].IsValid())
+	{
+		return false;
+	}
+	
+	UCompositeLayerPlate* Plate = Plates[0].Get();
+	TArray<UPrimitiveComponent*> CompositeMeshPrimitives = Plate->GetIncludedPrimitivesForActor(InActor);
+	
+	return !CompositeMeshPrimitives.IsEmpty();
+}
+
+bool FCompositeLayerPlateCustomization::IsCompositeMeshActorChecked(TSoftObjectPtr<AActor> InActor) const
+{
+	if (!InActor.IsValid())
+	{
+		return false;
+	}
+	
+	TSharedPtr<IDetailLayoutBuilder> DetailBuilder = CachedDetailBuilder.Pin();
+	if (!DetailBuilder.IsValid())
+	{
+		return false;
+	}
+	
+	TArray<TWeakObjectPtr<UCompositeLayerPlate>> Plates = DetailBuilder->GetObjectsOfTypeBeingCustomized<UCompositeLayerPlate>();
+	if (Plates.Num() != 1 || !Plates[0].IsValid())
+	{
+		return false;
+	}
+	
+	UCompositeLayerPlate* Plate = Plates[0].Get();
+	TArray<UPrimitiveComponent*> CompositeMeshPrimitives = Plate->GetIncludedPrimitivesForActor(InActor);
+	
+	return CompositeMeshPrimitives.IsEmpty();
+}
+
+void FCompositeLayerPlateCustomization::OnCompositeMeshPrimitiveToggled(TSoftObjectPtr<AActor> InActor, UPrimitiveComponent* InPrimitive)
+{
+	if (!InActor.IsValid())
+	{
+		return;
+	}
+	
+	TSharedPtr<IDetailLayoutBuilder> DetailBuilder = CachedDetailBuilder.Pin();
+	if (!DetailBuilder.IsValid())
+	{
+		return;
+	}
+	
+	TArray<TWeakObjectPtr<UCompositeLayerPlate>> Plates = DetailBuilder->GetObjectsOfTypeBeingCustomized<UCompositeLayerPlate>();
+	if (Plates.Num() != 1 || !Plates[0].IsValid())
+	{
+		return;
+	}
+	
+	UCompositeLayerPlate* Plate = Plates[0].Get();
+	TArray<UPrimitiveComponent*> CompositeMeshPrimitives = Plate->GetIncludedPrimitivesForActor(InActor);
+	
+	FScopedTransaction Transaction(LOCTEXT("ToggleCompositeMeshPrimitives", "Toggle Composite Mesh Primitive"));
+	Plate->Modify();
+	
+	const bool bInclude = !CompositeMeshPrimitives.Contains(InPrimitive);
+	Plate->SetPrimitiveIncluded(InPrimitive, bInclude);
+}
+
+bool FCompositeLayerPlateCustomization::IsCompositeMeshPrimitiveChecked(TSoftObjectPtr<AActor> InActor, UPrimitiveComponent* InPrimitive) const
+{
+	if (!InActor.IsValid())
+	{
+		return false;
+	}
+	
+	TSharedPtr<IDetailLayoutBuilder> DetailBuilder = CachedDetailBuilder.Pin();
+	if (!DetailBuilder.IsValid())
+	{
+		return false;
+	}
+	
+	TArray<TWeakObjectPtr<UCompositeLayerPlate>> Plates = DetailBuilder->GetObjectsOfTypeBeingCustomized<UCompositeLayerPlate>();
+	if (Plates.Num() != 1 || !Plates[0].IsValid())
+	{
+		return false;
+	}
+	
+	UCompositeLayerPlate* Plate = Plates[0].Get();
+	TArray<UPrimitiveComponent*> CompositeMeshPrimitives = Plate->GetIncludedPrimitivesForActor(InActor);
+	
+	return CompositeMeshPrimitives.Contains(InPrimitive);
+}
+
+void FCompositeLayerPlateCustomization::ApplyMaterialToActor(TSoftObjectPtr<AActor>& InActor, UMaterial* InMaterial)
+{
+	if (!InActor.IsValid())
+	{
+		return;
+	}
+	
+	TSharedPtr<IDetailLayoutBuilder> DetailBuilder = CachedDetailBuilder.Pin();
+	if (!DetailBuilder.IsValid())
+	{
+		return;
+	}
+	
+	TArray<TWeakObjectPtr<UCompositeLayerPlate>> Plates = DetailBuilder->GetObjectsOfTypeBeingCustomized<UCompositeLayerPlate>();
+	if (Plates.Num() != 1 || !Plates[0].IsValid())
+	{
+		return;
+	}
+	
+	UCompositeLayerPlate* Plate = Plates[0].Get();
+	TArray<UPrimitiveComponent*> CompositeMeshPrimitives = Plate->GetPrimitivesForActor(InActor);
+	
+	for (UPrimitiveComponent* Primitive : CompositeMeshPrimitives)
+	{
+		Primitive->Modify();
+		const int32 MaterialCount = Primitive->GetNumMaterials();
+		for (int32 MaterialIndex = 0; MaterialIndex < MaterialCount; MaterialIndex++)
+		{
+			Primitive->SetMaterial(MaterialIndex, InMaterial);
+		}
+	}
+}
+
 void FCompositeLayerPlateCustomization::CreateCompositeMeshActor()
 {
 	TSharedPtr<IDetailLayoutBuilder> PinnedDetailBuilder = CachedDetailBuilder.Pin();
diff --git a/Engine/Plugins/Compositing/Composite/Source/CompositeEditor/Private/Customizations/CompositeLayerPlateCustomization.h b/Engine/Plugins/Compositing/Composite/Source/CompositeEditor/Private/Customizations/CompositeLayerPlateCustomization.h
--- a/Engine/Plugins/Compositing/Composite/Source/CompositeEditor/Private/Customizations/CompositeLayerPlateCustomization.h
+++ b/Engine/Plugins/Compositing/Composite/Source/CompositeEditor/Private/Customizations/CompositeLayerPlateCustomization.h
@@ -3,6 +3,9 @@
 #pragma once
 
 #include "IDetailCustomization.h"
+#include "Materials/Material.h"
+#include "UObject/SoftObjectPtr.h"
+#include "Styling/SlateTypes.h"
 
 class SWidget;
 class IDetailPropertyRow;
@@ -45,9 +48,36 @@ private:
 	/** Opens the active media profile in the media profile editor */
 	void OpenMediaProfile();
 	
+	/** Creates the widget to display for each actor in the Components column of the actor picker table */
+	TSharedRef<SWidget> GetCompositeMeshColumnWidget(const TSoftObjectPtr<AActor>& InActor, const FName& InColumnName);
+
+	/** Gets the label to display for each actor in the Components column of the actor picker table */
+	FText GetCompositeMeshComponentsLabel(TSoftObjectPtr<AActor> InActor) const;
+
 	/** Raised by the composite mesh actor list when building the Add Actor menu */
 	void OnExtendCompositeMeshAddMenu(FMenuBuilder& MenuBuilder);
 
+	/** Raised by the composite mesh actor list when building the context menu */
+	void OnExtendCompositeMeshContextMenu(FMenuBuilder& MenuBuilder, TArray<TSoftObjectPtr<AActor>>& SelectedActors);
+
+	/** Raised when the composite mesh actor's entry in the Component section of the context menu is toggled */
+	void OnCompositeMeshActorToggled(TSoftObjectPtr<AActor> InActor);
+	
+	/** Gets whether the composite mesh actor's entry in the Component section of the context menu can be toggled */
+	bool CanToggleCompositeMeshActor(TSoftObjectPtr<AActor> InActor) const;
+	
+	/** Gets the check state for the composite mesh actor's entry in the Component section of the context menu */
+	bool IsCompositeMeshActorChecked(TSoftObjectPtr<AActor> InActor) const;
+	
+	/** Raised when a composite mesh primitive entry is toggled in the Component section of the context menu */
+	void OnCompositeMeshPrimitiveToggled(TSoftObjectPtr<AActor> InActor, UPrimitiveComponent* InPrimitive);
+	
+	/** Gets whether a composite mesh primitive entry is checked in the Component section of the context menu */
+	bool IsCompositeMeshPrimitiveChecked(TSoftObjectPtr<AActor> InActor, UPrimitiveComponent* InPrimitive) const;
+		
+	/** Raised when a material should be appplied to an actor in the actor picker list */
+	void ApplyMaterialToActor(TSoftObjectPtr<AActor>& InActor, UMaterial* InMaterial);
+
 	/** Creates a composite mesh actor at the appropriate position in the level and adds it to the plate's composite meshes list */
 	void CreateCompositeMeshActor();
 	
diff --git a/Engine/Plugins/Compositing/Composite/Source/CompositeEditor/Private/UI/SCompositeActorPickerTable.cpp b/Engine/Plugins/Compositing/Composite/Source/CompositeEditor/Private/UI/SCompositeActorPickerTable.cpp
--- a/Engine/Plugins/Compositing/Composite/Source/CompositeEditor/Private/UI/SCompositeActorPickerTable.cpp
+++ b/Engine/Plugins/Compositing/Composite/Source/CompositeEditor/Private/UI/SCompositeActorPickerTable.cpp
@@ -257,14 +257,17 @@ private:
 class SCompositeActorListItemRow : public SMultiColumnTableRow<SCompositeActorPickerTable::FActorListItemRef>
 {
 	using FActorListItemRef = SCompositeActorPickerTable::FActorListItemRef;
+	using FOnGenerateColumnWidget = SCompositeActorPickerTable::FOnGenerateColumnWidget;
 	
 public:
 	SLATE_BEGIN_ARGS(SCompositeActorListItemRow) { }
+		SLATE_EVENT(FOnGenerateColumnWidget, OnGenerateColumnWidget)
 	SLATE_END_ARGS()
 
 	void Construct(const FArguments& InArgs, const TSharedRef<STableViewBase>& InOwnerTable, const FActorListItemRef& InActorItem)
 	{
 		ActorItem = InActorItem;
+		OnGenerateColumnWidget = InArgs._OnGenerateColumnWidget;
 
 		STableRow<FActorListItemRef>::FArguments Args = FSuperRowType::FArguments()
 		   .Style(&FAppStyle::Get().GetWidgetStyle<FTableRowStyle>("SceneOutliner.TableViewRow"));
@@ -301,6 +304,10 @@ public:
 		{
 			return SNew(STextBlock).Text(this, &SCompositeActorListItemRow::GetActorType);
 		}
+		else if (OnGenerateColumnWidget.IsBound() && ActorItem.IsValid() && ActorItem->Actor.IsValid())
+		{
+			return OnGenerateColumnWidget.Execute(ActorItem->Actor, InColumnName);
+		}
 
 		return SNullWidget::NullWidget;
 	}
@@ -341,6 +348,7 @@ private:
 	
 private:
 	FActorListItemRef ActorItem;
+	FOnGenerateColumnWidget OnGenerateColumnWidget;
 };
 
 FCompositeActorPickerListRef::FCompositeActorPickerListRef(const TWeakObjectPtr<UObject>& InActorListOwner, const FName& InActorListPropertyName, TArray<TSoftObjectPtr<AActor>>* InActorList)
@@ -435,8 +443,10 @@ TSharedPtr<FSceneOutlinerFilters> SCompositeActorPickerTable::MakeDefaultSceneOu
 void SCompositeActorPickerTable::Construct(const FArguments& InArgs, const FCompositeActorPickerListRef& InActorListRef)
 {
 	ActorListRef = InActorListRef;
+	OnExtendContextMenu = InArgs._OnExtendContextMenu;
 	OnLayoutSizeChanged = InArgs._OnLayoutSizeChanged;
 	bShowApplyMaterialSection = InArgs._ShowApplyMaterialSection;
+	OnApplyMaterialToActor = InArgs._OnApplyMaterialToActor;
 	
 	FCoreUObjectDelegates::OnObjectPropertyChanged.AddSP(this, &SCompositeActorPickerTable::OnObjectPropertyChanged);
 	
@@ -448,6 +458,19 @@ void SCompositeActorPickerTable::Construct(const FArguments& InArgs, const FComp
 	CommandList = MakeShared<FUICommandList>();
 	BindCommands();
 	
+	TSharedRef<SHeaderRow> HeaderRow = 
+		SNew(SHeaderRow)
+
+		+SHeaderRow::Column(CompositeActorPickerTable::ActorListColumn_ActorName)
+		.FillWidth(0.6)
+		.DefaultLabel(LOCTEXT("ActorNameColumnLabel", "Item Label"))
+
+		+SHeaderRow::Column(CompositeActorPickerTable::ActorListColumn_Type)
+		.FillWidth(0.4)
+		.DefaultLabel(LOCTEXT("ActorTypeColumnLabel", "Type"));
+	
+	InArgs._OnGenerateHeaderColumns.ExecuteIfBound(HeaderRow);
+	
 	ChildSlot
 	[
 		SNew(SVerticalBox)
@@ -474,21 +497,12 @@ void SCompositeActorPickerTable::Construct(const FArguments& InArgs, const FComp
 			.OnDropped(this, &SCompositeActorPickerTable::OnListDropped)
 			[
 				SAssignNew(ListView, SListView<FActorListItemRef>)
-				.HeaderRow(
-					SNew(SHeaderRow)
-
-					+SHeaderRow::Column(CompositeActorPickerTable::ActorListColumn_ActorName)
-					.FillWidth(0.6)
-					.DefaultLabel(LOCTEXT("ActorNameColumnLabel", "Item Label"))
-
-					+SHeaderRow::Column(CompositeActorPickerTable::ActorListColumn_Type)
-					.FillWidth(0.4)
-					.DefaultLabel(LOCTEXT("ActorTypeColumnLabel", "Type"))
-				)
+				.HeaderRow(HeaderRow)
 				.ListItemsSource(&FilteredActorListItems)
-				.OnGenerateRow_Lambda([](FActorListItemRef InListItem, const TSharedRef<STableViewBase>& OwnerTable)
+				.OnGenerateRow_Lambda([OnGenerateColumnWidget = InArgs._OnGenerateColumnWidget](FActorListItemRef InListItem, const TSharedRef<STableViewBase>& OwnerTable)
 				{
-					return SNew(SCompositeActorListItemRow, OwnerTable, InListItem);
+					return SNew(SCompositeActorListItemRow, OwnerTable, InListItem)
+						.OnGenerateColumnWidget(OnGenerateColumnWidget);
 				})
 				.OnContextMenuOpening(this, &SCompositeActorPickerTable::CreateListContextMenu)
 			]
@@ -679,6 +693,18 @@ TSharedPtr<SWidget> SCompositeActorPickerTable::CreateListContextMenu()
 		MenuBuilder.AddMenuEntry(FCompositeEditorCommands::Get().RemoveActor);
 		MenuBuilder.EndSection();
 		
+		if (OnExtendContextMenu.IsBound())
+		{
+			TArray<TSoftObjectPtr<AActor>> SelectedActors;
+			TArray<FActorListItemRef> SelectedItems = ListView->GetSelectedItems();
+			Algo::Transform(SelectedItems, SelectedActors, [](const FActorListItemRef& Item)
+			{
+				return Item->Actor;
+			});
+			
+			OnExtendContextMenu.ExecuteIfBound(MenuBuilder, SelectedActors);
+		}
+		
 		return MenuBuilder.MakeWidget();
 	}
 
@@ -899,15 +925,29 @@ void SCompositeActorPickerTable::ApplyMaterialToActors(const TCHAR* InMaterialPa
 	for (int32 Index : ActorIndicesToEdit)
 	{
 		TSoftObjectPtr<AActor>& Actor = (*ActorListRef.ActorList)[Index];
-		for (UActorComponent* Component : Actor->GetComponents())
+		if (!Actor.IsValid())
 		{
-			if (UPrimitiveComponent* PrimitiveComponent = Cast<UPrimitiveComponent>(Component))
+			continue;
+		}
+		
+		// If the OnApplyMaterialToActor handler is bound, call that to allow custom apply logic. Otherwise, apply the material to 
+		// all the actor's primitive components
+		if (OnApplyMaterialToActor.IsBound())
+		{
+			OnApplyMaterialToActor.Execute(Actor, LoadedMaterial);
+		}
+		else
+		{
+			for (UActorComponent* Component : Actor->GetComponents())
 			{
-				PrimitiveComponent->Modify();
-				const int32 MaterialCount = PrimitiveComponent->GetNumMaterials();
-				for (int32 MaterialIndex = 0; MaterialIndex < MaterialCount; MaterialIndex++)
+				if (UPrimitiveComponent* PrimitiveComponent = Cast<UPrimitiveComponent>(Component))
 				{
-					PrimitiveComponent->SetMaterial(MaterialIndex, LoadedMaterial);
+					PrimitiveComponent->Modify();
+					const int32 MaterialCount = PrimitiveComponent->GetNumMaterials();
+					for (int32 MaterialIndex = 0; MaterialIndex < MaterialCount; MaterialIndex++)
+					{
+						PrimitiveComponent->SetMaterial(MaterialIndex, LoadedMaterial);
+					}
 				}
 			}
 		}
diff --git a/Engine/Plugins/Compositing/Composite/Source/CompositeEditor/Private/UI/SCompositeActorPickerTable.h b/Engine/Plugins/Compositing/Composite/Source/CompositeEditor/Private/UI/SCompositeActorPickerTable.h
--- a/Engine/Plugins/Compositing/Composite/Source/CompositeEditor/Private/UI/SCompositeActorPickerTable.h
+++ b/Engine/Plugins/Compositing/Composite/Source/CompositeEditor/Private/UI/SCompositeActorPickerTable.h
@@ -3,13 +3,17 @@
 #pragma once
 
 #include "EditorUndoClient.h"
+#include "UObject/SoftObjectPtr.h"
+#include "UObject/UnrealType.h"
 #include "Widgets/SCompoundWidget.h"
 
+class UMaterial;
 struct FSceneOutlinerFilters;
 class SCompositeActorPickerToolbar;
 class FMenuBuilder;
 class FUICommandList;
 template<typename T> class SListView;
+class SHeaderRow;
 
 /** Stores a pointer reference to a list of actors on a UObject to be managed by the actor picker */
 struct FCompositeActorPickerListRef
@@ -67,7 +71,11 @@ private:
 	using FActorListItemRef = TSharedPtr<FActorListItem>;
 
 public:
+	DECLARE_DELEGATE_OneParam(FOnGenerateHeaderColumns, TSharedRef<SHeaderRow> /* InHeaderRow */)
+	DECLARE_DELEGATE_RetVal_TwoParams(TSharedRef<SWidget>, FOnGenerateColumnWidget, const TSoftObjectPtr<AActor>& /* InActor */, const FName& /* ColumnName */)
 	DECLARE_DELEGATE_OneParam(FOnExtendAddMenu, FMenuBuilder&);
+	DECLARE_DELEGATE_TwoParams(FOnExtendContextMenu, FMenuBuilder&, TArray<TSoftObjectPtr<AActor>>& /*SelectedActors*/);
+	DECLARE_DELEGATE_TwoParams(FOnApplyMaterialToActor, TSoftObjectPtr<AActor>& /* InActor */, UMaterial* /* InMaterial */);
 
 public:
 	/** Returns default scene outliner filters that check IsAllowedPrimitiveClass and optionally exclude ACompositeMeshActor. */
@@ -75,9 +83,13 @@ public:
 
 	SLATE_BEGIN_ARGS(SCompositeActorPickerTable) : _ShowApplyMaterialSection(false) { }
 		SLATE_ATTRIBUTE(TSharedPtr<FSceneOutlinerFilters>, SceneOutlinerFilters)
+		SLATE_EVENT(FOnGenerateHeaderColumns, OnGenerateHeaderColumns)
+		SLATE_EVENT(FOnGenerateColumnWidget, OnGenerateColumnWidget)
 		SLATE_EVENT(FOnExtendAddMenu, OnExtendAddMenu)
+		SLATE_EVENT(FOnExtendContextMenu, OnExtendContextMenu)
 		SLATE_EVENT(FSimpleDelegate, OnLayoutSizeChanged)
 		SLATE_ATTRIBUTE(bool, ShowApplyMaterialSection)
+		SLATE_EVENT(FOnApplyMaterialToActor, OnApplyMaterialToActor)
 	SLATE_END_ARGS()
 
 	void Construct(const FArguments& InArgs, const FCompositeActorPickerListRef& InActorListRef);
@@ -159,12 +171,18 @@ private:
 	/** Raised when building the Add menu, allows the menu contents to be extended with additional entries */
 	FOnExtendAddMenu OnExtendAddMenu;
 
+	/** Raised when building the context menu for entries in the table, allows the context menu contents to be extended with additional entries */
+	FOnExtendContextMenu OnExtendContextMenu;
+
 	/** Raised when the number of items in the list view changes */
 	FSimpleDelegate OnLayoutSizeChanged;
 
 	/** Flag to enable the material application section for selected actors. */
 	TAttribute<bool> bShowApplyMaterialSection;
 
+	/** Raised when a material is being applied to an actor. Allows for custom logic on how the material is applied */
+	FOnApplyMaterialToActor OnApplyMaterialToActor;
+
 	friend class SCompositeActorPickerToolbar;
 	friend class SCompositeActorListItemRow;
 };
