Composure: Adding new sky sphere / sky light convenience actor to easily use the plate texture for image-based lighting.

We also add a convenience button on the plate layer to create the actor and assign it the current plate texture.

#virtualized
#Synced-CL 52595449

diff --git a/Engine/Plugins/Compositing/Composite/Source/Composite/Private/CompositeSkySphereActor.cpp b/Engine/Plugins/Compositing/Composite/Source/Composite/Private/CompositeSkySphereActor.cpp
new file mode 100644
--- /dev/null
+++ b/Engine/Plugins/Compositing/Composite/Source/Composite/Private/CompositeSkySphereActor.cpp
@@ -0,0 +1,182 @@
+// Copyright Epic Games, Inc. All Rights Reserved.
+
+#include "CompositeSkySphereActor.h"
+
+#include "Components/SkyLightComponent.h"
+#include "Components/StaticMeshComponent.h"
+#include "Engine/StaticMesh.h"
+#include "Engine/Texture.h"
+#include "Engine/Texture2D.h"
+#include "Materials/MaterialInterface.h"
+#include "Materials/MaterialInstanceDynamic.h"
+#include "MediaTexture.h"
+#include "Profile/IMediaProfileManager.h"
+#include "Profile/MediaProfile.h"
+#include "Profile/MediaProfilePlaybackManager.h"
+#include "UObject/ConstructorHelpers.h"
+
+ACompositeSkySphereActor::ACompositeSkySphereActor(const FObjectInitializer& ObjectInitializer)
+	: Super(ObjectInitializer)
+{
+	struct FConstructorStatics
+	{
+		ConstructorHelpers::FObjectFinder<UStaticMesh> SkySphere;
+		ConstructorHelpers::FObjectFinder<UMaterialInterface> SkySphere_Material;
+		ConstructorHelpers::FObjectFinder<UTexture2D> DefaultTexture;
+		FConstructorStatics()
+			: SkySphere(TEXT("/Composite/Meshes/SM_CompositeSkySphere.SM_CompositeSkySphere"))
+			, SkySphere_Material(TEXT("/Composite/Materials/M_CompositeSkySphere.M_CompositeSkySphere"))
+			, DefaultTexture(TEXT("/Composite/Textures/T_Composite_SMPTE_Color_Bars_16x9.T_Composite_SMPTE_Color_Bars_16x9"))
+		{
+		}
+	};
+	static FConstructorStatics ConstructorStatics;
+
+	SkySphereComponent = CreateDefaultSubobject<UStaticMeshComponent>(TEXT("SkySphereComponent"));
+	SetRootComponent(SkySphereComponent);
+	SkySphereComponent->SetCastShadow(false);
+	SkySphereComponent->SetCollisionEnabled(ECollisionEnabled::NoCollision);
+	SkySphereComponent->SetEnableGravity(false);
+	// SM_CompositeSkySphere has a 100 UU diameter (50 UU radius).
+	// Scale 1000 gives a 50,000 UU radius — comfortably beyond the default 15,000 UU sky distance threshold.
+	SkySphereComponent->SetRelativeScale3D(FVector(1000.0f));
+
+	if (ensure(ConstructorStatics.SkySphere.Object != nullptr))
+	{
+		SkySphereComponent->SetStaticMesh(ConstructorStatics.SkySphere.Object);
+	}
+
+	if (ConstructorStatics.SkySphere_Material.Object != nullptr)
+	{
+		SkySphereComponent->SetMaterial(0, ConstructorStatics.SkySphere_Material.Object);
+	}
+
+	Texture = Cast<UTexture>(ConstructorStatics.DefaultTexture.Object);
+
+	SkyLightComponent = CreateDefaultSubobject<USkyLightComponent>(TEXT("SkyLightComponent"));
+	SkyLightComponent->SetupAttachment(SkySphereComponent);
+	SkyLightComponent->bRealTimeCapture = true;
+	SkyLightComponent->bLowerHemisphereIsBlack = false;
+}
+
+ACompositeSkySphereActor::~ACompositeSkySphereActor() = default;
+
+void ACompositeSkySphereActor::OnConstruction(const FTransform& Transform)
+{
+	Super::OnConstruction(Transform);
+	
+	UpdateMaterial();
+}
+
+void ACompositeSkySphereActor::PostLoad()
+{
+	Super::PostLoad();
+
+	if (!HasAnyFlags(RF_ClassDefaultObject))
+	{
+		UpdateMaterial();
+		TryOpenMediaProfileSource();
+	}
+}
+
+void ACompositeSkySphereActor::EndPlay(const EEndPlayReason::Type EndPlayReason)
+{
+	TryCloseMediaProfileSource();
+	Super::EndPlay(EndPlayReason);
+}
+
+void ACompositeSkySphereActor::UpdateMaterial()
+{
+	UMaterialInterface* CurrentMaterial = SkySphereComponent->GetMaterial(0);
+	if (CurrentMaterial == nullptr)
+	{
+		return;
+	}
+
+	UMaterialInstanceDynamic* MID = Cast<UMaterialInstanceDynamic>(CurrentMaterial);
+	if (!MID)
+	{
+		if (Texture == nullptr)
+		{
+			return; // No texture and no existing MID — nothing to do.
+		}
+		MID = SkySphereComponent->CreateAndSetMaterialInstanceDynamic(0);
+	}
+
+	if (ensure(MID != nullptr))
+	{
+		MID->SetTextureParameterValue(TextureParameterName, Texture);
+	}
+}
+
+void ACompositeSkySphereActor::SetTexture(UTexture* InTexture)
+{
+	TryCloseMediaProfileSource();
+	Texture = InTexture;
+	UpdateMaterial();
+	TryOpenMediaProfileSource();
+}
+
+void ACompositeSkySphereActor::TryOpenMediaProfileSource()
+{
+	UMediaTexture* MediaTexture = Cast<UMediaTexture>(Texture);
+	if (!MediaTexture)
+	{
+		return;
+	}
+
+	UMediaProfile* ActiveMediaProfile = IMediaProfileManager::Get().GetCurrentMediaProfile();
+	if (!ActiveMediaProfile)
+	{
+		return;
+	}
+
+	int32 MediaSourceIndex = INDEX_NONE;
+	if (ActiveMediaProfile->GetPlaybackManager()->IsValidSourceMediaTexture(MediaTexture, MediaSourceIndex))
+	{
+		ActiveMediaProfile->GetPlaybackManager()->OpenSourceFromIndex(MediaSourceIndex, this);
+	}
+}
+
+void ACompositeSkySphereActor::TryCloseMediaProfileSource()
+{
+	UMediaProfile* ActiveMediaProfile = IMediaProfileManager::Get().GetCurrentMediaProfile();
+	if (!ActiveMediaProfile)
+	{
+		return;
+	}
+
+	UMediaProfilePlaybackManager::FCloseSourceArgs Args;
+	Args.Consumer = this;
+	ActiveMediaProfile->GetPlaybackManager()->CloseSourcesForConsumer(Args);
+}
+
+#if WITH_EDITOR
+void ACompositeSkySphereActor::PreEditChange(FProperty* PropertyThatWillChange)
+{
+	Super::PreEditChange(PropertyThatWillChange);
+
+	if (PropertyThatWillChange &&
+		PropertyThatWillChange->GetFName() == GET_MEMBER_NAME_CHECKED(ACompositeSkySphereActor, Texture))
+	{
+		TryCloseMediaProfileSource();
+	}
+}
+
+void ACompositeSkySphereActor::PostEditChangeProperty(FPropertyChangedEvent& PropertyChangedEvent)
+{
+	Super::PostEditChangeProperty(PropertyChangedEvent);
+
+	const FName PropertyName = PropertyChangedEvent.GetMemberPropertyName();
+	if (PropertyName == GET_MEMBER_NAME_CHECKED(ACompositeSkySphereActor, Texture) ||
+		PropertyName == GET_MEMBER_NAME_CHECKED(ACompositeSkySphereActor, TextureParameterName))
+	{
+		UpdateMaterial();
+	}
+
+	if (PropertyName == GET_MEMBER_NAME_CHECKED(ACompositeSkySphereActor, Texture))
+	{
+		TryOpenMediaProfileSource();
+	}
+}
+#endif
diff --git a/Engine/Plugins/Compositing/Composite/Source/Composite/Public/CompositeSkySphereActor.h b/Engine/Plugins/Compositing/Composite/Source/Composite/Public/CompositeSkySphereActor.h
new file mode 100644
--- /dev/null
+++ b/Engine/Plugins/Compositing/Composite/Source/Composite/Public/CompositeSkySphereActor.h
@@ -0,0 +1,80 @@
+// Copyright Epic Games, Inc. All Rights Reserved.
+
+#pragma once
+
+#include "CoreTypes.h"
+#include "GameFramework/Actor.h"
+
+#include "CompositeSkySphereActor.generated.h"
+
+#define UE_API COMPOSITE_API
+
+class USkyLightComponent;
+class UStaticMeshComponent;
+class UTexture;
+
+/**
+ * Convenience actor providing image-based lighting from a plate texture.
+ * Wraps an inverted sphere mesh and a real-time skylight. Assign Texture and
+ * a sky sphere material with a "CompositeTexture" parameter to get image-
+ * based lighting from footage.
+ */
+UCLASS(MinimalAPI)
+class ACompositeSkySphereActor : public AActor
+{
+	GENERATED_BODY()
+
+public:
+	UE_API ACompositeSkySphereActor(const FObjectInitializer& ObjectInitializer);
+	UE_API ~ACompositeSkySphereActor();
+
+	//~ Begin AActor Interface
+	UE_API virtual void OnConstruction(const FTransform& Transform) override;
+	UE_API virtual void PostLoad() override;
+	UE_API virtual void EndPlay(const EEndPlayReason::Type EndPlayReason) override;
+	//~ End AActor Interface
+
+	/** Sets the sky sphere texture and immediately updates the material. */
+	UFUNCTION(BlueprintCallable, Category = "Texture")
+	UE_API void SetTexture(UTexture* InTexture);
+
+	/** Returns the current sky sphere texture. */
+	UFUNCTION(BlueprintPure, Category = "Texture")
+	UE_API UTexture* GetTexture() const { return Texture; }
+
+	/** Name of the texture parameter in the sky sphere material to receive Texture. */
+	UPROPERTY(EditAnywhere, BlueprintReadWrite, Category = "Texture")
+	FName TextureParameterName = TEXT("CompositeTexture");
+
+#if WITH_EDITOR
+	UE_API virtual void PreEditChange(FProperty* PropertyThatWillChange) override;
+	UE_API virtual void PostEditChangeProperty(struct FPropertyChangedEvent& PropertyChangedEvent) override;
+#endif
+
+private:
+	friend class FCompositeSkySphereActorCustomization;
+
+	/** Updates the sky sphere MID with the current Texture value. */
+	void UpdateMaterial();
+
+	/** If Texture is a MediaTexture from the active Media Profile, registers this actor as a consumer. */
+	void TryOpenMediaProfileSource();
+
+	/** Unregisters this actor as a consumer of any Media Profile source it previously opened. */
+	void TryCloseMediaProfileSource();
+
+	/** Texture to drive the sky sphere material's "CompositeTexture" parameter. */
+	UPROPERTY(EditAnywhere, Category = "Texture",
+		meta = (AllowPrivateAccess = "true", AllowedClasses = "/Script/Engine.Texture2D,/Script/Engine.Texture2DDynamic,/Script/Engine.TextureRenderTarget2D,/Script/MediaAssets.MediaTexture"))
+	TObjectPtr<UTexture> Texture;
+
+	/** Sky sphere mesh component. */
+	UPROPERTY(VisibleAnywhere, BlueprintReadOnly, Category = "Composite", meta = (AllowPrivateAccess = "true"))
+	TObjectPtr<UStaticMeshComponent> SkySphereComponent;
+
+	/** Real-time skylight that captures the sky sphere for IBL. */
+	UPROPERTY(VisibleAnywhere, BlueprintReadOnly, Category = "Composite", meta = (AllowPrivateAccess = "true"))
+	TObjectPtr<USkyLightComponent> SkyLightComponent;
+};
+
+#undef UE_API
diff --git a/Engine/Plugins/Compositing/Composite/Source/CompositeEditor/CompositeEditor.Build.cs b/Engine/Plugins/Compositing/Composite/Source/CompositeEditor/CompositeEditor.Build.cs
--- a/Engine/Plugins/Compositing/Composite/Source/CompositeEditor/CompositeEditor.Build.cs
+++ b/Engine/Plugins/Compositing/Composite/Source/CompositeEditor/CompositeEditor.Build.cs
@@ -34,6 +34,7 @@ public class CompositeEditor : ModuleRules
 				"LevelSequence",
 				"MediaAssets",
 				"MediaFrameworkUtilities",
+				"MediaFrameworkUtilitiesEditor",
 				"MovieScene",
 				"MovieSceneTracks",
 				"MovieSceneTools",
diff --git a/Engine/Plugins/Compositing/Composite/Source/CompositeEditor/Private/ActorFactories/CompositeSkySphereActorFactory.cpp b/Engine/Plugins/Compositing/Composite/Source/CompositeEditor/Private/ActorFactories/CompositeSkySphereActorFactory.cpp
new file mode 100644
--- /dev/null
+++ b/Engine/Plugins/Compositing/Composite/Source/CompositeEditor/Private/ActorFactories/CompositeSkySphereActorFactory.cpp
@@ -0,0 +1,31 @@
+// Copyright Epic Games, Inc. All Rights Reserved.
+
+#include "CompositeSkySphereActorFactory.h"
+
+#include "CompositeAnalytics.h"
+#include "CompositeSkySphereActor.h"
+
+#include UE_INLINE_GENERATED_CPP_BY_NAME(CompositeSkySphereActorFactory)
+
+#define LOCTEXT_NAMESPACE "CompositeSkySphereActorFactory"
+
+UCompositeSkySphereActorFactory::UCompositeSkySphereActorFactory(const FObjectInitializer& ObjInit)
+	: Super(ObjInit)
+{
+	DisplayName = LOCTEXT("DisplayName", "Composite Sky Sphere Actor");
+	NewActorClass = ACompositeSkySphereActor::StaticClass();
+}
+
+void UCompositeSkySphereActorFactory::PostSpawnActor(UObject* Asset, AActor* NewActor)
+{
+	Super::PostSpawnActor(Asset, NewActor);
+
+	ACompositeSkySphereActor* SkySphereActor = Cast<ACompositeSkySphereActor>(NewActor);
+
+	if (IsValid(SkySphereActor) && !SkySphereActor->IsTemplate() && !SkySphereActor->HasAnyFlags(RF_Transient))
+	{
+		Composite::Analytics::RecordActorAdded(*SkySphereActor);
+	}
+}
+
+#undef LOCTEXT_NAMESPACE
diff --git a/Engine/Plugins/Compositing/Composite/Source/CompositeEditor/Private/ActorFactories/CompositeSkySphereActorFactory.h b/Engine/Plugins/Compositing/Composite/Source/CompositeEditor/Private/ActorFactories/CompositeSkySphereActorFactory.h
new file mode 100644
--- /dev/null
+++ b/Engine/Plugins/Compositing/Composite/Source/CompositeEditor/Private/ActorFactories/CompositeSkySphereActorFactory.h
@@ -0,0 +1,23 @@
+// Copyright Epic Games, Inc. All Rights Reserved.
+
+#pragma once
+
+#include "UObject/ObjectMacros.h"
+#include "ActorFactories/ActorFactory.h"
+
+#include "CompositeSkySphereActorFactory.generated.h"
+
+UCLASS()
+class UCompositeSkySphereActorFactory : public UActorFactory
+{
+	GENERATED_BODY()
+
+public:
+	/** Constructor */
+	UCompositeSkySphereActorFactory(const FObjectInitializer& ObjInit);
+
+public:
+	//~ Begin UActorFactory Interface
+	virtual void PostSpawnActor(UObject* Asset, AActor* NewActor) override;
+	//~ End UActorFactory Interface
+};
diff --git a/Engine/Plugins/Compositing/Composite/Source/CompositeEditor/Private/CompositeEditorModule.cpp b/Engine/Plugins/Compositing/Composite/Source/CompositeEditor/Private/CompositeEditorModule.cpp
--- a/Engine/Plugins/Compositing/Composite/Source/CompositeEditor/Private/CompositeEditorModule.cpp
+++ b/Engine/Plugins/Compositing/Composite/Source/CompositeEditor/Private/CompositeEditorModule.cpp
@@ -17,6 +17,7 @@
 #include "Components/CompositeDepthMeshComponent.h"
 #include "Components/CompositeViewProjectionComponent.h"
 #include "CompositeActor.h"
+#include "CompositeSkySphereActor.h"
 #include "CompositeEditorCommands.h"
 #include "CompositeEditorStyle.h"
 #include "PropertyEditorModule.h"
@@ -34,6 +35,7 @@
 #include "Customizations/CompositeLayerSingleLightShadowCustomization.h"
 #include "Customizations/CompositeDepthMeshComponentCustomization.h"
 #include "Customizations/CompositePassColorKeyerCustomization.h"
+#include "Customizations/CompositeSkySphereActorCustomization.h"
 #include "Layers/CompositeLayerBase.h"
 #include "Layers/CompositeLayerMainRender.h"
 #include "Layers/CompositeLayerPlate.h"
@@ -268,6 +270,11 @@ void FCompositeEditorModule::RegisterCustomizations()
 		UCompositePassColorKeyer::StaticClass()->GetFName(),
 		FOnGetDetailCustomizationInstance::CreateStatic(&FCompositePassColorKeyerCustomization::MakeInstance)
 	);
+
+	PropertyModule.RegisterCustomClassLayout(
+		ACompositeSkySphereActor::StaticClass()->GetFName(),
+		FOnGetDetailCustomizationInstance::CreateStatic(&FCompositeSkySphereActorCustomization::MakeInstance)
+	);
 }
 
 void FCompositeEditorModule::UnregisterCustomizations()
@@ -288,6 +295,7 @@ void FCompositeEditorModule::UnregisterCustomizations()
 			PropertyModule->UnregisterCustomClassLayout(UCompositeViewProjectionComponent::StaticClass()->GetFName());
 			PropertyModule->UnregisterCustomClassLayout(UCompositeDepthMeshComponent::StaticClass()->GetFName());
 			PropertyModule->UnregisterCustomClassLayout(UCompositePassColorKeyer::StaticClass()->GetFName());
+			PropertyModule->UnregisterCustomClassLayout(ACompositeSkySphereActor::StaticClass()->GetFName());
 		}
 	}
 }
diff --git a/Engine/Plugins/Compositing/Composite/Source/CompositeEditor/Private/CompositeEditorStyle.cpp b/Engine/Plugins/Compositing/Composite/Source/CompositeEditor/Private/CompositeEditorStyle.cpp
--- a/Engine/Plugins/Compositing/Composite/Source/CompositeEditor/Private/CompositeEditorStyle.cpp
+++ b/Engine/Plugins/Compositing/Composite/Source/CompositeEditor/Private/CompositeEditorStyle.cpp
@@ -36,6 +36,7 @@ FCompositeEditorStyle::FCompositeEditorStyle()
 	Set("CompositeEditor.Composure", new CORE_IMAGE_BRUSH_SVG("Starship/Common/ComposureCompositing", Icon16x16));
 
 	Set("ClassIcon.CompositeActor", new CORE_IMAGE_BRUSH_SVG("Starship/Common/ComposureCompositing", Icon16x16));
+	Set("ClassIcon.CompositeSkySphereActor", new CORE_IMAGE_BRUSH_SVG("Starship/AssetIcons/SkyLight_16", Icon16x16));
 	Set("ClassIcon.CompositeLayerMainRender", new IMAGE_BRUSH_SVG("Editor/Icons/MainRender", Icon16x16));
 	Set("ClassIcon.CompositeLayerShadowReflection", new IMAGE_BRUSH_SVG("Editor/Icons/Shadow", Icon16x16));
 	Set("ClassIcon.CompositeLayerSingleLightShadow", new IMAGE_BRUSH_SVG("Editor/Icons/Shadow", Icon16x16));
diff --git a/Engine/Plugins/Compositing/Composite/Source/CompositeEditor/Private/Customizations/CompositeLayerPlateCustomization.cpp b/Engine/Plugins/Compositing/Composite/Source/CompositeEditor/Private/Customizations/CompositeLayerPlateCustomization.cpp
--- a/Engine/Plugins/Compositing/Composite/Source/CompositeEditor/Private/Customizations/CompositeLayerPlateCustomization.cpp
+++ b/Engine/Plugins/Compositing/Composite/Source/CompositeEditor/Private/Customizations/CompositeLayerPlateCustomization.cpp
@@ -5,6 +5,7 @@
 #include "CompositeActor.h"
 #include "CompositeCoreSettings.h"
 #include "CompositeEditorStyle.h"
+#include "CompositeSkySphereActor.h"
 #include "CompositeMeshActor.h"
 #include "DetailCategoryBuilder.h"
 #include "DetailLayoutBuilder.h"
@@ -14,20 +15,18 @@
 #include "IDetailGroup.h"
 #include "IPropertyUtilities.h"
 #include "LevelEditorSubsystem.h"
-#include "MediaSource.h"
 #include "MediaTexture.h"
-#include "PropertyCustomizationHelpers.h"
+#include "Misc/MessageDialog.h"
 #include "ScopedTransaction.h"
+#include "SMediaProfileSourceTexturePicker.h"
 #include "UnrealEdGlobals.h"
 #include "Camera/CameraComponent.h"
 #include "Editor/UnrealEdEngine.h"
+#include "Engine/SkyLight.h"
+#include "EngineUtils.h"
 #include "Framework/MultiBox/MultiBoxBuilder.h"
 #include "Layers/CompositeLayerPlate.h"
-#include "Profile/IMediaProfileManager.h"
-#include "Profile/MediaProfile.h"
-#include "Profile/MediaProfilePlaybackManager.h"
 #include "Styling/SlateIconFinder.h"
-#include "Subsystems/AssetEditorSubsystem.h"
 #include "UI/CompositePlatePassListOwner.h"
 #include "UI/SCompositeActorPickerTable.h"
 #include "UI/SCompositePassTreePanel.h"
@@ -192,134 +191,37 @@ void FCompositeLayerPlateCustomization::CustomizeTexturePropertyRow(const TShare
 	]
 	.ValueContent()
 	[
-		SNew(SObjectPropertyEntryBox)
-		.PropertyHandle(InPropertyHandle)
-		.AllowedClass(UTexture::StaticClass())
-		.AllowCreate(true)
+		SNew(SMediaProfileSourceTexturePicker)
+		.TexturePropertyHandle(InPropertyHandle)
 		.ThumbnailPool(CachedDetailBuilder.IsValid() ? CachedDetailBuilder.Pin()->GetThumbnailPool() : nullptr)
+		.AllowedClass(UTexture::StaticClass())
 		.OnShouldFilterAsset_Lambda([](FAssetData const& AssetData)
 		{
 			if (UClass* AssetClass = AssetData.GetClass())
 			{
 				return !(AssetClass->IsChildOf(UTexture2D::StaticClass()) || AssetClass->IsChildOf(UMediaTexture::StaticClass()));
 			}
-
 			return true;
 		})
-		.CustomContentSlot()
+		.AdditionalContent()
 		[
-			SNew(SHorizontalBox)
-
-			+SHorizontalBox::Slot()
-			.AutoWidth()
+			SNew(SComboButton)
+			.HasDownArrow(true)
+			.ComboButtonStyle(FAppStyle::Get(), "SimpleComboButton")
+			.ToolTipText(LOCTEXT("PlaceSkySphereButtonToolTip", "Place a Composite Sky Sphere Actor using the plate's current texture"))
+			.HAlign(HAlign_Center)
+			.VAlign(VAlign_Center)
+			.OnGetMenuContent(this, &FCompositeLayerPlateCustomization::GetSkySphereActorMenu)
+			.IsEnabled(this, &FCompositeLayerPlateCustomization::CanPlaceCompositeSkySphereActor)
+			.ButtonContent()
 			[
-				SNew(SComboButton)
-				.HasDownArrow(true)
-				.ComboButtonStyle(FAppStyle::Get(), "SimpleComboButton")
-				.ToolTipText(LOCTEXT("MediaProfileSourceSelectorToolTip", "Select media texture from a media source in the active Media Profile"))
-				.HAlign(HAlign_Center)
-				.VAlign(VAlign_Center)
-				.OnGetMenuContent(this, &FCompositeLayerPlateCustomization::GetMediaProfileSourceSelectorMenu)
-				.IsEnabled(this, &FCompositeLayerPlateCustomization::HasActiveMediaProfile)
-				.ButtonContent()
-				[
-					SNew(SImage)
-					.Image(FCompositeEditorStyle::Get().GetBrush("CompositeEditor.MediaProfile"))
-				]
+				SNew(SImage)
+				.Image(FCompositeEditorStyle::Get().GetBrush("ClassIcon.CompositeSkySphereActor"))
 			]
 		]
 	];
 }
 
-bool FCompositeLayerPlateCustomization::HasActiveMediaProfile() const
-{
-	UMediaProfile* ActiveMediaProfile = IMediaProfileManager::Get().GetCurrentMediaProfile();
-	return IsValid(ActiveMediaProfile);
-}
-
-TSharedRef<SWidget> FCompositeLayerPlateCustomization::GetMediaProfileSourceSelectorMenu()
-{
-	UMediaProfile* ActiveMediaProfile = IMediaProfileManager::Get().GetCurrentMediaProfile();
-	if (!ActiveMediaProfile)
-	{
-		return SNullWidget::NullWidget;
-	}
-	
-	const bool bCloseAfterSelection = true;
-	FMenuBuilder MenuBuilder(bCloseAfterSelection, nullptr);
-
-	MenuBuilder.BeginSection(NAME_None, LOCTEXT("MediaProfileSourcesSectionLabel", "Texture from Media Profile"));
-	{
-		bool bNoMediaSources = true;
-		for (int32 Index = 0; Index < ActiveMediaProfile->NumMediaSources(); ++Index)
-		{
-			if (UMediaSource* MediaSource = ActiveMediaProfile->GetMediaSource(Index))
-			{
-				bNoMediaSources = false;
-
-				MenuBuilder.AddMenuEntry(
-					FText::FromString(ActiveMediaProfile->GetLabelForMediaSource(Index)),
-					FText::GetEmpty(),
-					FSlateIconFinder::FindIconForClass(MediaSource->GetClass()),
-					FUIAction(FExecuteAction::CreateSP(this, &FCompositeLayerPlateCustomization::SetTextureToMediaProfileSource, Index))
-				);
-			}
-		}
-
-		if (bNoMediaSources)
-		{
-			TSharedRef<SWidget> TextBlock = SNew(STextBlock)
-				.Text(LOCTEXT("NoMediaSources", "No Media Sources"))
-				.Font(FCoreStyle::GetDefaultFontStyle("Italic", 10));
-			
-			MenuBuilder.AddWidget(TextBlock,FText::GetEmpty());
-		}
-	}
-	MenuBuilder.EndSection();
-
-	MenuBuilder.AddSeparator();
-	MenuBuilder.BeginSection(NAME_None);
-	{
-		MenuBuilder.AddMenuEntry(
-			FText::Format(LOCTEXT("OpenActiveMediaProfileLabelFormat", "Open {0}"), FText::FromString(ActiveMediaProfile->GetName())),
-			FText::GetEmpty(),
-			FSlateIconFinder::FindIconForClass(ActiveMediaProfile->StaticClass()),
-			FUIAction(FExecuteAction::CreateSP(this, &FCompositeLayerPlateCustomization::OpenMediaProfile))
-		);
-	}
-	MenuBuilder.EndSection();
-	
-	return MenuBuilder.MakeWidget();
-}
-
-void FCompositeLayerPlateCustomization::SetTextureToMediaProfileSource(int32 InMediaSourceIndex)
-{
-	UMediaProfile* ActiveMediaProfile = IMediaProfileManager::Get().GetCurrentMediaProfile();
-	if (!ActiveMediaProfile)
-	{
-		return;
-	}
-
-	if (!IsValid(ActiveMediaProfile->GetMediaSource(InMediaSourceIndex)))
-	{
-		return;
-	}
-
-	UMediaTexture* MediaTexture = ActiveMediaProfile->GetPlaybackManager()->GetSourceMediaTextureFromIndex(InMediaSourceIndex);
-	TexturePropertyHandle->SetValue(MediaTexture);
-}
-
-void FCompositeLayerPlateCustomization::OpenMediaProfile()
-{
-	UMediaProfile* ActiveMediaProfile = IMediaProfileManager::Get().GetCurrentMediaProfile();
-	if (!ActiveMediaProfile)
-	{
-		return;
-	}
-	
-	GEditor->GetEditorSubsystem<UAssetEditorSubsystem>()->OpenEditorForAsset(ActiveMediaProfile);
-}
-
 TSharedRef<SWidget> FCompositeLayerPlateCustomization::GetCompositeMeshColumnWidget(const TSoftObjectPtr<AActor>& InActor, const FName& InColumnName)
 {
 	if (InColumnName == CompositeLayerPlateCustomization::CompositeMeshListColumn_Components)
@@ -716,6 +618,102 @@ bool FCompositeLayerPlateCustomization::CanCreateCompositeMeshActor() const
 	return true;
 }
 
+TSharedRef<SWidget> FCompositeLayerPlateCustomization::GetSkySphereActorMenu()
+{
+	const bool bCloseAfterSelection = true;
+	FMenuBuilder MenuBuilder(bCloseAfterSelection, nullptr);
+
+	MenuBuilder.AddMenuEntry(
+		LOCTEXT("PlaceSkySphereActorEntryLabel", "Place Composite SkySphere Actor"),
+		FText::GetEmpty(),
+		FSlateIcon(FCompositeEditorStyle::Get().GetStyleSetName(), "ClassIcon.CompositeSkySphereActor"),
+		FUIAction(FExecuteAction::CreateSP(this, &FCompositeLayerPlateCustomization::PlaceCompositeSkySphereActor))
+	);
+
+	return MenuBuilder.MakeWidget();
+}
+
+void FCompositeLayerPlateCustomization::PlaceCompositeSkySphereActor()
+{
+	TSharedPtr<IDetailLayoutBuilder> PinnedDetailBuilder = CachedDetailBuilder.Pin();
+	if (!PinnedDetailBuilder.IsValid())
+	{
+		return;
+	}
+
+	TArray<TWeakObjectPtr<UCompositeLayerPlate>> Objects = PinnedDetailBuilder->GetObjectsOfTypeBeingCustomized<UCompositeLayerPlate>();
+	if (Objects.Num() != 1 || !Objects[0].IsValid())
+	{
+		return;
+	}
+
+	ACompositeActor* CompositeActor = Objects[0]->GetTypedOuter<ACompositeActor>();
+	if (!CompositeActor)
+	{
+		return;
+	}
+
+	// Resolve the currently assigned plate texture.
+	UObject* TextureObj = nullptr;
+	if (TexturePropertyHandle.IsValid())
+	{
+		TexturePropertyHandle->GetValue(TextureObj);
+	}
+	UTexture* CurrentTexture = Cast<UTexture>(TextureObj);
+
+	UWorld* World = CompositeActor->GetWorld();
+	if (!World)
+	{
+		return;
+	}
+
+	// Only prompt if the scene already has a sky light — the sky sphere adds its own real-time sky light.
+	if (TActorIterator<ASkyLight>(World))
+	{
+		const FText DialogMessage = LOCTEXT("PlaceSkySphereConfirmSkyLightWarning",
+			"The scene already contains a Sky Light actor. The Composite Sky Sphere Actor includes its own real-time sky light, which may produce unexpected results.");
+
+		if (FMessageDialog::Open(EAppMsgType::YesNo, DialogMessage,
+			LOCTEXT("PlaceSkySphereDialogTitle", "Create Composite Sky Sphere")) != EAppReturnType::Yes)
+		{
+			return;
+		}
+	}
+
+	FScopedTransaction Transaction(LOCTEXT("PlaceSkySphereTransaction", "Place Composite Sky Sphere"));
+
+	// Use deferred spawn so the texture is set before OnConstruction runs UpdateMaterial.
+	ACompositeSkySphereActor* SkySphereActor = World->SpawnActorDeferred<ACompositeSkySphereActor>(
+		ACompositeSkySphereActor::StaticClass(), FTransform::Identity);
+	if (SkySphereActor)
+	{
+		SkySphereActor->SetTexture(CurrentTexture);
+		SkySphereActor->FinishSpawning(FTransform::Identity);
+
+		if (GEditor)
+		{
+			GEditor->SelectNone(true, true);
+			GEditor->SelectActor(SkySphereActor, true, true);
+		}
+	}
+}
+
+bool FCompositeLayerPlateCustomization::CanPlaceCompositeSkySphereActor() const
+{
+	if (!CanCreateCompositeMeshActor())
+	{
+		return false;
+	}
+
+	UObject* TextureObj = nullptr;
+	if (TexturePropertyHandle.IsValid())
+	{
+		TexturePropertyHandle->GetValue(TextureObj);
+	}
+
+	return Cast<UTexture>(TextureObj) != nullptr;
+}
+
 void FCompositeLayerPlateCustomization::OnLayoutSizeChanged()
 {
 	TSharedPtr<IDetailLayoutBuilder> PinnedDetailBuilder = CachedDetailBuilder.Pin();
diff --git a/Engine/Plugins/Compositing/Composite/Source/CompositeEditor/Private/Customizations/CompositeLayerPlateCustomization.h b/Engine/Plugins/Compositing/Composite/Source/CompositeEditor/Private/Customizations/CompositeLayerPlateCustomization.h
--- a/Engine/Plugins/Compositing/Composite/Source/CompositeEditor/Private/Customizations/CompositeLayerPlateCustomization.h
+++ b/Engine/Plugins/Compositing/Composite/Source/CompositeEditor/Private/Customizations/CompositeLayerPlateCustomization.h
@@ -36,18 +36,6 @@ private:
 	/** Customizes the Texture property in UCompositeLayerPlate to add in the media profile source selector */
 	void CustomizeTexturePropertyRow(const TSharedPtr<IPropertyHandle>& InPropertyHandle, IDetailPropertyRow& InPropertyRow);
 
-	/** Gets whether there is an active media profile configured in the editor */
-	bool HasActiveMediaProfile() const;
-
-	/** Gets the dropdown menu widget to display when the media profile button is pressed */
-	TSharedRef<SWidget> GetMediaProfileSourceSelectorMenu();
-
-	/** Sets the Texture property handle's value to the media texture corresponding to the specified media source in the active media profile */
-	void SetTextureToMediaProfileSource(int32 InMediaSourceIndex);
-
-	/** Opens the active media profile in the media profile editor */
-	void OpenMediaProfile();
-	
 	/** Creates the widget to display for each actor in the Components column of the actor picker table */
 	TSharedRef<SWidget> GetCompositeMeshColumnWidget(const TSoftObjectPtr<AActor>& InActor, const FName& InColumnName);
 
@@ -80,10 +68,19 @@ private:
 
 	/** Creates a composite mesh actor at the appropriate position in the level and adds it to the plate's composite meshes list */
 	void CreateCompositeMeshActor();
-	
+
 	/** Gets whether a composite mesh actor can be created for this plate layer */
 	bool CanCreateCompositeMeshActor() const;
 
+	/** Gets the dropdown menu widget to display when the sky sphere button is pressed */
+	TSharedRef<SWidget> GetSkySphereActorMenu();
+
+	/** Spawns a CompositeSkySphereActor in the current level */
+	void PlaceCompositeSkySphereActor();
+
+	/** Gets whether a sky sphere actor can be placed */
+	bool CanPlaceCompositeSkySphereActor() const;
+
 	/** Raised when the custom widgets' layout size may have changed */
 	void OnLayoutSizeChanged();
 	
diff --git a/Engine/Plugins/Compositing/Composite/Source/CompositeEditor/Private/Customizations/CompositeSkySphereActorCustomization.cpp b/Engine/Plugins/Compositing/Composite/Source/CompositeEditor/Private/Customizations/CompositeSkySphereActorCustomization.cpp
new file mode 100644
--- /dev/null
+++ b/Engine/Plugins/Compositing/Composite/Source/CompositeEditor/Private/Customizations/CompositeSkySphereActorCustomization.cpp
@@ -0,0 +1,169 @@
+// Copyright Epic Games, Inc. All Rights Reserved.
+
+#include "CompositeSkySphereActorCustomization.h"
+
+#include "CompositeSkySphereActor.h"
+#include "DetailCategoryBuilder.h"
+#include "DetailLayoutBuilder.h"
+#include "Engine/Texture2D.h"
+#include "Engine/Texture2DDynamic.h"
+#include "Engine/TextureRenderTarget2D.h"
+#include "MediaTexture.h"
+#include "Modules/ModuleManager.h"
+#include "PropertyEditorModule.h"
+#include "SMediaProfileSourceTexturePicker.h"
+
+#define LOCTEXT_NAMESPACE "FCompositeSkySphereActorCustomization"
+
+TSharedRef<IDetailCustomization> FCompositeSkySphereActorCustomization::MakeInstance()
+{
+	return MakeShared<FCompositeSkySphereActorCustomization>();
+}
+
+void FCompositeSkySphereActorCustomization::CustomizeDetails(const TSharedPtr<IDetailLayoutBuilder>& DetailBuilder)
+{
+	CachedDetailBuilder = DetailBuilder;
+	IDetailCustomization::CustomizeDetails(DetailBuilder);
+}
+
+void FCompositeSkySphereActorCustomization::CustomizeDetails(IDetailLayoutBuilder& DetailLayout)
+{
+	CustomizeSections();
+
+	// Texture category: pinned at Important priority.
+	IDetailCategoryBuilder& TextureCategory = DetailLayout.EditCategory(
+		TEXT("Texture"), FText::GetEmpty(), ECategoryPriority::Important);
+
+	// Texture property — use SMediaProfileSourceTexturePicker with our broader allowed types.
+	TSharedRef<IPropertyHandle> TextureHandle =
+		DetailLayout.GetProperty(GET_MEMBER_NAME_CHECKED(ACompositeSkySphereActor, Texture));
+	TextureCategory.AddProperty(TextureHandle).CustomWidget()
+	.NameContent()
+	[
+		TextureHandle->CreatePropertyNameWidget()
+	]
+	.ValueContent()
+	[
+		SNew(SMediaProfileSourceTexturePicker)
+		.TexturePropertyHandle(TextureHandle)
+		.ThumbnailPool(CachedDetailBuilder.IsValid() ? CachedDetailBuilder.Pin()->GetThumbnailPool() : nullptr)
+		.AllowedClass(UTexture::StaticClass())
+		.OnShouldFilterAsset_Lambda([](const FAssetData& AssetData)
+		{
+			if (UClass* AssetClass = AssetData.GetClass())
+			{
+				return !(AssetClass->IsChildOf(UTexture2D::StaticClass()) ||
+				         AssetClass->IsChildOf(UTexture2DDynamic::StaticClass()) ||
+				         AssetClass->IsChildOf(UTextureRenderTarget2D::StaticClass()) ||
+				         AssetClass->IsChildOf(UMediaTexture::StaticClass()));
+			}
+			return true;
+		})
+	];
+
+	TextureCategory.AddProperty(
+		DetailLayout.GetProperty(GET_MEMBER_NAME_CHECKED(ACompositeSkySphereActor, TextureParameterName)));
+}
+
+void FCompositeSkySphereActorCustomization::CustomizeSections()
+{
+	FPropertyEditorModule& PropertyModule =
+		FModuleManager::GetModuleChecked<FPropertyEditorModule>(TEXT("PropertyEditor"));
+	const FName ClassName = ACompositeSkySphereActor::StaticClass()->GetFName();
+
+	// Custom pills — registered before the suppress blocks so they appear first.
+	{
+		TSharedRef<FPropertySection> Section = PropertyModule.FindOrCreateSection(
+			ClassName, TEXT("General"), LOCTEXT("SectionGeneral", "General"));
+		Section->AddCategory(TEXT("Texture"));
+		Section->AddCategory(TEXT("Ray Tracing"));
+	}
+	{
+		TSharedRef<FPropertySection> Section = PropertyModule.FindOrCreateSection(
+			ClassName, TEXT("Light"), LOCTEXT("SectionLight", "Light"));
+		Section->AddCategory(TEXT("Light"));
+	}
+	{
+		TSharedRef<FPropertySection> Section = PropertyModule.FindOrCreateSection(
+			ClassName, TEXT("Materials"), LOCTEXT("SectionMaterials", "Materials"));
+		Section->AddCategory(TEXT("Materials"));
+	}
+	{
+		TSharedRef<FPropertySection> Section = PropertyModule.FindOrCreateSection(
+			ClassName, TEXT("Ray Tracing"), LOCTEXT("SectionRayTracing", "Ray Tracing"));
+		Section->AddCategory(TEXT("Ray Tracing"));
+	}
+	{
+		TSharedRef<FPropertySection> Section = PropertyModule.FindOrCreateSection(
+			ClassName, TEXT("Static Mesh"), LOCTEXT("SectionStaticMesh", "Static Mesh"));
+		Section->AddCategory(TEXT("Static Mesh"));
+	}
+	{
+		TSharedRef<FPropertySection> Section = PropertyModule.FindOrCreateSection(
+			ClassName, TEXT("Texture"), LOCTEXT("SectionTexture", "Texture"));
+		Section->AddCategory(TEXT("Texture"));
+	}
+
+	// Suppress pills irrelevant for a sky sphere actor.
+	// Removing all categories from a section makes its pill disappear.
+	// Displaced categories remain accessible via "All".
+	{
+		TSharedRef<FPropertySection> Section = PropertyModule.FindOrCreateSection(
+			ClassName, TEXT("Actor"), LOCTEXT("SectionActorSuppress", "Actor"));
+		Section->RemoveCategory(TEXT("Actor"));
+	}
+	{
+		TSharedRef<FPropertySection> Section = PropertyModule.FindOrCreateSection(
+			ClassName, TEXT("Misc"), LOCTEXT("SectionMiscSuppress", "Misc"));
+		Section->RemoveCategory(TEXT("Asset User Data"));
+		Section->RemoveCategory(TEXT("Cooking"));
+		Section->RemoveCategory(TEXT("Input"));
+		Section->RemoveCategory(TEXT("Navigation"));
+		Section->RemoveCategory(TEXT("Performance"));
+		Section->RemoveCategory(TEXT("Replication"));
+		Section->RemoveCategory(TEXT("Tags"));
+	}
+	{
+		TSharedRef<FPropertySection> Section = PropertyModule.FindOrCreateSection(
+			ClassName, TEXT("Streaming"), LOCTEXT("SectionStreamingSuppress", "Streaming"));
+		Section->RemoveCategory(TEXT("Data Layers"));
+		Section->RemoveCategory(TEXT("HLOD"));
+		Section->RemoveCategory(TEXT("World Partition"));
+	}
+	// Rendering: individual categories are exposed via dedicated pills (Materials, Ray Tracing).
+	{
+		TSharedRef<FPropertySection> Section = PropertyModule.FindOrCreateSection(
+			ClassName, TEXT("Rendering"), LOCTEXT("SectionRenderingSuppress", "Rendering"));
+		Section->RemoveCategory(TEXT("Lighting"));
+		Section->RemoveCategory(TEXT("Lightmass"));
+		Section->RemoveCategory(TEXT("Materials"));
+		Section->RemoveCategory(TEXT("Mobile"));
+		Section->RemoveCategory(TEXT("Ray Tracing"));
+		Section->RemoveCategory(TEXT("Path Tracing"));
+		Section->RemoveCategory(TEXT("Rendering"));
+		Section->RemoveCategory(TEXT("Texture Streaming"));
+		Section->RemoveCategory(TEXT("Virtual Texture"));
+		Section->RemoveCategory(TEXT("Material Parameters"));
+		Section->RemoveCategory(TEXT("Mesh Painting"));
+		// From SkyLightComponent:
+		Section->RemoveCategory(TEXT("Light"));
+		Section->RemoveCategory(TEXT("Atmosphere and Cloud"));
+		Section->RemoveCategory(TEXT("Distance Field Ambient Occlusion"));
+	}
+	// LOD: sky sphere is a full-screen sphere — LOD settings are not meaningful.
+	{
+		TSharedRef<FPropertySection> Section = PropertyModule.FindOrCreateSection(
+			ClassName, TEXT("LOD"), LOCTEXT("SectionLODSuppress", "LOD"));
+		Section->RemoveCategory(TEXT("HLOD"));
+		Section->RemoveCategory(TEXT("LOD"));
+	}
+	// Physics: collision is explicitly disabled in the constructor.
+	{
+		TSharedRef<FPropertySection> Section = PropertyModule.FindOrCreateSection(
+			ClassName, TEXT("Physics"), LOCTEXT("SectionPhysicsSuppress", "Physics"));
+		Section->RemoveCategory(TEXT("Collision"));
+		Section->RemoveCategory(TEXT("Physics"));
+	}
+}
+
+#undef LOCTEXT_NAMESPACE
diff --git a/Engine/Plugins/Compositing/Composite/Source/CompositeEditor/Private/Customizations/CompositeSkySphereActorCustomization.h b/Engine/Plugins/Compositing/Composite/Source/CompositeEditor/Private/Customizations/CompositeSkySphereActorCustomization.h
new file mode 100644
--- /dev/null
+++ b/Engine/Plugins/Compositing/Composite/Source/CompositeEditor/Private/Customizations/CompositeSkySphereActorCustomization.h
@@ -0,0 +1,37 @@
+// Copyright Epic Games, Inc. All Rights Reserved.
+
+#pragma once
+
+#include "IDetailCustomization.h"
+
+class IDetailPropertyRow;
+class IPropertyHandle;
+class SWidget;
+
+/**
+ * Details customization for ACompositeSkySphereActor.
+ *
+ * Registers section filter pills (General | Light | Materials | Ray Tracing | Static Mesh | Texture | All)
+ * via FPropertyEditorModule::FindOrCreateSection so they appear in the native UE5 pill row at the top of
+ * the Details panel.
+ *
+ * Also adds a media profile source picker to the Texture property, matching the plate layer UX.
+ */
+class FCompositeSkySphereActorCustomization : public IDetailCustomization
+{
+public:
+	/** Makes a new instance of the details customization */
+	static TSharedRef<IDetailCustomization> MakeInstance();
+
+	//~ Begin IDetailCustomization Interface
+	virtual void CustomizeDetails(const TSharedPtr<IDetailLayoutBuilder>& DetailBuilder) override;
+	virtual void CustomizeDetails(IDetailLayoutBuilder& DetailLayout) override;
+	//~ End IDetailCustomization Interface
+
+private:
+	/** Registers section filter pills via FPropertyEditorModule::FindOrCreateSection. */
+	void CustomizeSections();
+
+	/** Cached layout builder for thumbnail pool access. */
+	TWeakPtr<IDetailLayoutBuilder> CachedDetailBuilder;
+};
