Composure: Converted existing color keyer pass into regular pixel shader, since we're no longer expecting to use compute features directly.

#Synced-CL 52595449

diff --git a/Engine/Plugins/Compositing/Composite/Shaders/Private/CompositeColorKeyer.usf b/Engine/Plugins/Compositing/Composite/Shaders/Private/CompositeColorKeyer.usf
--- a/Engine/Plugins/Compositing/Composite/Shaders/Private/CompositeColorKeyer.usf
+++ b/Engine/Plugins/Compositing/Composite/Shaders/Private/CompositeColorKeyer.usf
@@ -5,18 +5,12 @@
 
 #include "CompositeColorKeyerCommon.ush"
 
-#ifndef TILE_SIZE
-    #error TILE_SIZE must be defined
-#endif
-
 Texture2D Input_Texture;
 SamplerState Input_Sampler;
 
 Texture2D CleanPlate_Texture;
 SamplerState CleanPlate_Sampler;
 
-RWTexture2D<float4> RWOutputTexture;
-
 SCREEN_PASS_TEXTURE_VIEWPORT(Input)
 SCREEN_PASS_TEXTURE_VIEWPORT(Output)
 FScreenTransform SvPositionToInputTextureUV;
@@ -56,21 +50,10 @@ float4 Output(float4 Color)
 	return Premultiply(Color);
 }
 
-[numthreads(TILE_SIZE, TILE_SIZE, 1)]
-void MainCS(
-    uint2 GroupId : SV_GroupID,
-    uint2 GroupThreadId : SV_GroupThreadID,
-    uint GroupThreadIndex : SV_GroupIndex,
-    uint2 DispatchThreadId : SV_DispatchThreadID)
+void MainPS(
+    float4 SvPosition : SV_POSITION,
+    out float4 OutColor : SV_Target0)
 {
-	float4 SvPosition = float4((float2)DispatchThreadId + Output_ViewportMin + 0.5f, 0.0f, 1.0f);
-	float2 OutputUV = (SvPosition.xy) * Output_ExtentInverse;
-
-	if (IsComputeUVOutOfBounds(OutputUV))
-	{
-		return;
-	}
-
 	// Input parameters
 	const float3 ChannelWeights = Params0.xyz;
 	const float DespillStrength = Params0.w;
@@ -78,7 +61,7 @@ void MainCS(
 	const float AlphaMax = Params1.y;
 	const float DevignetteStrength = Params1.z;
 	float2 InputUV = ApplyScreenTransform(SvPosition.xy, SvPositionToInputTextureUV);
-	
+
 	float4 Foreground = Input_Texture.SampleLevel(Input_Sampler, InputUV, 0);
 #if USE_CLEAN_PLATE
 	float3 Background = CleanPlate_Texture.SampleLevel(CleanPlate_Sampler, InputUV, 0).rgb;
@@ -140,7 +123,6 @@ void MainCS(
 	// Apply final alpha threshold
 	Result.a = Linearstep(AlphaMin, AlphaMax, Result.a);
 #endif
-	
-	const uint2 PixelPos = DispatchThreadId + Output_ViewportMin;
-	RWOutputTexture[PixelPos] = Output(Result);
+
+	OutColor = Output(Result);
 }
diff --git a/Engine/Plugins/Compositing/Composite/Source/Composite/Private/Passes/CompositePassColorKeyer.cpp b/Engine/Plugins/Compositing/Composite/Source/Composite/Private/Passes/CompositePassColorKeyer.cpp
--- a/Engine/Plugins/Compositing/Composite/Source/Composite/Private/Passes/CompositePassColorKeyer.cpp
+++ b/Engine/Plugins/Compositing/Composite/Source/Composite/Private/Passes/CompositePassColorKeyer.cpp
@@ -51,7 +51,6 @@ class FCompositePassColorKeyerShader : public FGlobalShader
 	BEGIN_SHADER_PARAMETER_STRUCT(FParameters, )
 		SHADER_PARAMETER_STRUCT(FScreenPassTextureInput, Input)
 		SHADER_PARAMETER_STRUCT(FScreenPassTextureInput, CleanPlate)
-		SHADER_PARAMETER_RDG_TEXTURE_UAV(RWTexture2D<float4>, RWOutputTexture)
 		SHADER_PARAMETER_STRUCT(FScreenPassTextureViewportParameters, Output)
 		SHADER_PARAMETER(FScreenTransform, SvPositionToInputTextureUV)
 		SHADER_PARAMETER(FVector3f, KeyColor)
@@ -61,20 +60,12 @@ class FCompositePassColorKeyerShader : public FGlobalShader
 		SHADER_PARAMETER(uint32, Visualization)
 		SHADER_PARAMETER(uint32, bPreserveVignetteAfterKey)
 		SHADER_PARAMETER(uint32, bInvertAlpha)
+		RENDER_TARGET_BINDING_SLOTS()
 	END_SHADER_PARAMETER_STRUCT()
-
-	static inline void ModifyCompilationEnvironment(const FGlobalShaderPermutationParameters& Parameters, FShaderCompilerEnvironment& OutEnvironment)
-	{
-		OutEnvironment.SetDefine(TEXT("TILE_SIZE"), ThreadGroupSize);
-
-		FGlobalShader::ModifyCompilationEnvironment(Parameters, OutEnvironment);
-	}
-
-	static const uint32 ThreadGroupSize = 16;
 };
 
 IMPLEMENT_GLOBAL_SHADER(FCompositePassDenoiseShader, "/Plugin/Composite/Private/CompositeDenoise.usf", "MainCS", SF_Compute);
-IMPLEMENT_GLOBAL_SHADER(FCompositePassColorKeyerShader, "/Plugin/Composite/Private/CompositeColorKeyer.usf", "MainCS", SF_Compute);
+IMPLEMENT_GLOBAL_SHADER(FCompositePassColorKeyerShader, "/Plugin/Composite/Private/CompositeColorKeyer.usf", "MainPS", SF_Pixel);
 
 static FScreenPassTexture AddColorKeyerDenoisePass(FRDGBuilder& GraphBuilder, const FScreenPassTexture& Input, ERHIFeatureLevel::Type FeatureLevel)
 {
@@ -131,29 +122,13 @@ namespace UE::Composite::Private
 
 		FGlobalShaderMap* GlobalShaderMap = GetGlobalShaderMap(InView.GetFeatureLevel());
 
-		/** Write the keyer CS result to a UAV-only intermediate rather than directly to Output.Texture.
-		 *  Output.Texture is RenderTargetable; D3D12 state tracking for UAV writes onto RenderTargetable
-		 *  textures can race with the RHI thread, corrupting the subsequent graphics SRV read in the
-		 *  merge pass. See CompositePassDilation.cpp and GeometryMaskPostProcess_Blur.cpp for the pattern. */
-		const FIntPoint OutputSize = Output.ViewRect.Size();
-		const FRDGTextureDesc KeyerComputeDesc = FRDGTextureDesc::Create2D(OutputSize, Output.Texture->Desc.Format, FClearValueBinding::Black, TexCreate_UAV | TexCreate_ShaderResource);
-		FRDGTextureRef KeyerComputeTexture = GraphBuilder.CreateTexture(KeyerComputeDesc, TEXT("CompositeColorKeyerCS"));
-
-		/** Zero-based viewport matching the CS dispatch and UAV write coordinates.
-		 *  Both OutputParameters and SvPositionToInputTextureUV must use the same viewport
-		 *  so that the TexelPosition -> ViewportUV basis is consistent with dispatch thread IDs. */
-		const FScreenPassTextureViewport KeyerViewport(FScreenPassTexture(KeyerComputeTexture, FIntRect(FIntPoint::ZeroValue, OutputSize)));
-
-		const FScreenPassTextureViewportParameters OutputParameters = GetScreenPassTextureViewportParameters(KeyerViewport);
-
 		FRHISamplerState* InputSampler = Metadata.GetSamplerState();
 
 		FCompositePassColorKeyerShader::FParameters* PassParameters = GraphBuilder.AllocParameters<FCompositePassColorKeyerShader::FParameters>();
 		PassParameters->Input = GetScreenPassTextureInput(Input, InputSampler);
-		PassParameters->RWOutputTexture = GraphBuilder.CreateUAV(KeyerComputeTexture);
-		PassParameters->Output = OutputParameters;
+		PassParameters->Output = GetScreenPassTextureViewportParameters(FScreenPassTextureViewport(Output));
 		PassParameters->SvPositionToInputTextureUV = (
-			FScreenTransform::ChangeTextureBasisFromTo(KeyerViewport, FScreenTransform::ETextureBasis::TexelPosition, FScreenTransform::ETextureBasis::ViewportUV) *
+			FScreenTransform::ChangeTextureBasisFromTo(FScreenPassTextureViewport(Output), FScreenTransform::ETextureBasis::TexelPosition, FScreenTransform::ETextureBasis::ViewportUV) *
 			FScreenTransform::ChangeTextureBasisFromTo(FScreenPassTextureViewport(Input), FScreenTransform::ETextureBasis::ViewportUV, FScreenTransform::ETextureBasis::TextureUV));
 		PassParameters->KeyColor = FVector3f(KeyColor);
 		PassParameters->Params0 = Params0;
@@ -162,6 +137,7 @@ namespace UE::Composite::Private
 		PassParameters->Visualization = static_cast<uint32>(Visualization);
 		PassParameters->bPreserveVignetteAfterKey = static_cast<uint32>(bPreserveVignetteAfterKey);
 		PassParameters->bInvertAlpha = static_cast<uint32>(bInvertAlpha);
+		PassParameters->RenderTargets[0] = Output.GetRenderTargetBinding();
 
 		FRDGTextureRef CleanPlateResource = nullptr;
 		PassParameters->CleanPlate.Texture = GSystemTextures.GetBlackDummy(GraphBuilder);
@@ -189,14 +165,15 @@ namespace UE::Composite::Private
 		PermutationVector.Set<FCompositePassColorKeyerShader::FUseCleanPlate>(CleanPlateResource != nullptr);
 		PermutationVector.Set<FCompositePassColorKeyerShader::FDespillOnly>(bDespillOnly);
 
-		TShaderMapRef<FCompositePassColorKeyerShader> ComputeShader(GlobalShaderMap, PermutationVector);
-		FComputeShaderUtils::AddPass(
+		const FIntPoint OutputSize = Output.ViewRect.Size();
+		TShaderMapRef<FCompositePassColorKeyerShader> PixelShader(GlobalShaderMap, PermutationVector);
+		FPixelShaderUtils::AddFullscreenPass(
 			GraphBuilder,
+			GlobalShaderMap,
 			RDG_EVENT_NAME("Composite.ColorKeyer (%dx%d)", OutputSize.X, OutputSize.Y),
-			ERDGPassFlags::Compute,
-			ComputeShader,
+			PixelShader,
 			PassParameters,
-			FComputeShaderUtils::GetGroupCount(OutputSize, FCompositePassColorKeyerShader::ThreadGroupSize)
+			Output.ViewRect
 		);
 
 		if (CleanPlateResource != nullptr)
@@ -204,9 +181,6 @@ namespace UE::Composite::Private
 			GraphBuilder.UseExternalAccessMode(CleanPlateResource, ERHIAccess::SRVMask);
 		}
 
-		/** Copy the CS result to the actual output render target, bridging UAV -> CopyDest -> SRVGraphics. */
-		AddCopyTexturePass(GraphBuilder, KeyerComputeTexture, Output.Texture);
-
 		return UE::CompositeCore::FPassTexture{ MoveTemp(Output), Metadata };
 	}
 }
