Composure: Added basic support for luminance keying with a new Luma Keyer pass type, which recycles the existing proxy/shader.

#Synced-CL 52595449

diff --git a/Engine/Plugins/Compositing/Composite/Shaders/Private/CompositeColorKeyer.usf b/Engine/Plugins/Compositing/Composite/Shaders/Private/CompositeColorKeyer.usf
--- a/Engine/Plugins/Compositing/Composite/Shaders/Private/CompositeColorKeyer.usf
+++ b/Engine/Plugins/Compositing/Composite/Shaders/Private/CompositeColorKeyer.usf
@@ -18,6 +18,8 @@ FScreenTransform SvPositionToInputTextureUV;
 float3 KeyColor;
 float4 Params0;
 float4 Params1;
+uint KeyerMode;
+uint bDisplaySpaceLuminance;
 uint ScreenType;
 uint Visualization;
 uint bInvertAlpha;
@@ -50,16 +52,8 @@ float4 Output(float4 Color)
 	return Premultiply(Color);
 }
 
-void MainPS(
-    float4 SvPosition : SV_POSITION,
-    out float4 OutColor : SV_Target0)
+float4 MainPS(float4 SvPosition : SV_POSITION) : SV_Target0
 {
-	// Input parameters
-	const float3 ChannelWeights = Params0.xyz;
-	const float DespillStrength = Params0.w;
-	const float AlphaMin = Params1.x;
-	const float AlphaMax = Params1.y;
-	const float DevignetteStrength = Params1.z;
 	float2 InputUV = ApplyScreenTransform(SvPosition.xy, SvPositionToInputTextureUV);
 
 	float4 Foreground = Input_Texture.SampleLevel(Input_Sampler, InputUV, 0);
@@ -74,28 +68,46 @@ void MainPS(
 #if !USE_DESPILL_ONLY
 	float2 ScreenPos = UVToScreenPos(InputUV, Input_Extent);
 	float AspectRatio = Output_ViewportSize.y * Output_ViewportSizeInverse.x;
+	float DevignetteStrength = Params1.z;
+	float AlphaMin = Params1.x;
+	float AlphaMax = Params1.y;
 	float VignetteMask = Vignette(ScreenPos, AspectRatio, DevignetteStrength);
 
-	if(DevignetteStrength > ZeroEpsilon)
-	{
-		Result.rgb *= rcp(VignetteMask);
-	}
-
 	BRANCH
-	if(ScreenType == 0)
-	{
-		Result = RedScreenKeyer(Result, Background, ChannelWeights[1], ChannelWeights[2]);
-	}
-	else if(ScreenType == 1)
+	if (KeyerMode == 1) // Luma key: derive alpha from pixel luminance.
 	{
-		Result = GreenScreenKeyer(Result, Background, ChannelWeights[0], ChannelWeights[2]);
+		// C++ packs LumaRange into Params1.xy when KeyerMode == Luma.
+		float3 LumaInput = bDisplaySpaceLuminance ? LinearToSrgb(Foreground.rgb) : Foreground.rgb;
+		float Luma = Luminance(LumaInput);
+		Result.a = Linearstep(AlphaMin, AlphaMax, Luma) * Foreground.a;
 	}
-	else if(ScreenType == 2)
+	else // Chroma key (KeyerMode == 0)
 	{
-		Result = BlueScreenKeyer(Result, Background, ChannelWeights[0], ChannelWeights[1]);
+		float3 ChannelWeights = Params0.xyz;
+
+		if(DevignetteStrength > ZeroEpsilon)
+		{
+			Result.rgb *= rcp(VignetteMask);
+		}
+
+		BRANCH
+		if(ScreenType == 0)
+		{
+			Result = RedScreenKeyer(Result, Background, ChannelWeights[1], ChannelWeights[2]);
+		}
+		else if(ScreenType == 1)
+		{
+			Result = GreenScreenKeyer(Result, Background, ChannelWeights[0], ChannelWeights[2]);
+		}
+		else if(ScreenType == 2)
+		{
+			Result = BlueScreenKeyer(Result, Background, ChannelWeights[0], ChannelWeights[1]);
+		}
 	}
 #endif
 
+	float DespillStrength = Params0.w;
+
 	BRANCH
 	if(DespillStrength > 0.0f)
 	{
@@ -115,14 +127,20 @@ void MainPS(
 	}
 
 #if !USE_DESPILL_ONLY
-	if(bPreserveVignetteAfterKey > 0 && DevignetteStrength > ZeroEpsilon)
+	// Chroma-only post-processing: vignette re-application and alpha threshold.
+	// In luma mode DespillStrength and DevignetteStrength are packed as 0, so both branches
+	// are natural no-ops, but we still gate on KeyerMode to make the intent explicit.
+	if (KeyerMode == 0)
 	{
-		Result.rgb *= VignetteMask;
-	}
+		if(bPreserveVignetteAfterKey > 0 && DevignetteStrength > ZeroEpsilon)
+		{
+			Result.rgb *= VignetteMask;
+		}
 
-	// Apply final alpha threshold
-	Result.a = Linearstep(AlphaMin, AlphaMax, Result.a);
+		// Apply final alpha threshold
+		Result.a = Linearstep(AlphaMin, AlphaMax, Result.a);
+	}
 #endif
 
-	OutColor = Output(Result);
+	return Output(Result);
 }
diff --git a/Engine/Plugins/Compositing/Composite/Source/Composite/Private/Passes/CompositePassColorKeyer.cpp b/Engine/Plugins/Compositing/Composite/Source/Composite/Private/Passes/CompositePassColorKeyer.cpp
--- a/Engine/Plugins/Compositing/Composite/Source/Composite/Private/Passes/CompositePassColorKeyer.cpp
+++ b/Engine/Plugins/Compositing/Composite/Source/Composite/Private/Passes/CompositePassColorKeyer.cpp
@@ -56,6 +56,8 @@ class FCompositePassColorKeyerShader : public FGlobalShader
 		SHADER_PARAMETER(FVector3f, KeyColor)
 		SHADER_PARAMETER(FVector4f, Params0)
 		SHADER_PARAMETER(FVector4f, Params1)
+		SHADER_PARAMETER(uint32, KeyerMode)
+		SHADER_PARAMETER(uint32, bDisplaySpaceLuminance)
 		SHADER_PARAMETER(uint32, ScreenType)
 		SHADER_PARAMETER(uint32, Visualization)
 		SHADER_PARAMETER(uint32, bPreserveVignetteAfterKey)
@@ -121,6 +123,7 @@ namespace UE::Composite::Private
 		}
 
 		FGlobalShaderMap* GlobalShaderMap = GetGlobalShaderMap(InView.GetFeatureLevel());
+		const FIntPoint OutputSize = Output.ViewRect.Size();
 
 		FRHISamplerState* InputSampler = Metadata.GetSamplerState();
 
@@ -133,6 +136,8 @@ namespace UE::Composite::Private
 		PassParameters->KeyColor = FVector3f(KeyColor);
 		PassParameters->Params0 = Params0;
 		PassParameters->Params1 = Params1;
+		PassParameters->KeyerMode = static_cast<uint32>(KeyerMode);
+		PassParameters->bDisplaySpaceLuminance = static_cast<uint32>(bDisplaySpaceLuminance);
 		PassParameters->ScreenType = static_cast<uint32>(ScreenType);
 		PassParameters->Visualization = static_cast<uint32>(Visualization);
 		PassParameters->bPreserveVignetteAfterKey = static_cast<uint32>(bPreserveVignetteAfterKey);
@@ -165,7 +170,6 @@ namespace UE::Composite::Private
 		PermutationVector.Set<FCompositePassColorKeyerShader::FUseCleanPlate>(CleanPlateResource != nullptr);
 		PermutationVector.Set<FCompositePassColorKeyerShader::FDespillOnly>(bDespillOnly);
 
-		const FIntPoint OutputSize = Output.ViewRect.Size();
 		TShaderMapRef<FCompositePassColorKeyerShader> PixelShader(GlobalShaderMap, PermutationVector);
 		FPixelShaderUtils::AddFullscreenPass(
 			GraphBuilder,
@@ -303,7 +307,7 @@ void UCompositePassColorKeyer::PostEditChangeProperty(FPropertyChangedEvent& Pro
 		case ECompositeColorKeyerScreenType::Blue:
 			KeyColor = FLinearColor::Blue;
 			break;
-		
+
 		default:
 			checkNoEntry();
 			break;
@@ -312,3 +316,31 @@ void UCompositePassColorKeyer::PostEditChangeProperty(FPropertyChangedEvent& Pro
 }
 #endif
 
+UCompositePassLumaKeyer::UCompositePassLumaKeyer(const FObjectInitializer& ObjectInitializer)
+	: Super(ObjectInitializer)
+	, LumaRange(0.0f, 1.0f)
+	, Visualization(ECompositeColorKeyerVisualization::Key)
+	, bInvertAlpha(false)
+	, bDisplaySpaceLuminance(true)
+	, DenoiseMethod(ECompositeDenoiseMethod::None)
+{
+}
+
+UCompositePassLumaKeyer::~UCompositePassLumaKeyer() = default;
+
+FCompositeCorePassProxy* UCompositePassLumaKeyer::GetProxy(const UE::CompositeCore::FPassInputDecl& InputDecl, FCompositeTraversalContext& InContext, FSceneRenderingBulkObjectAllocator& InFrameAllocator) const
+{
+	using namespace UE::Composite::Private;
+
+	FCompositePassColorKeyerProxy* Proxy = InFrameAllocator.Create<FCompositePassColorKeyerProxy>(UE::CompositeCore::FPassInputDeclArray{ InputDecl });
+	Proxy->KeyerMode = ECompositeColorKeyerMode::Luma;
+	Proxy->DenoiseMethod = DenoiseMethod;
+	Proxy->bInvertAlpha = bInvertAlpha;
+	Proxy->Visualization = Visualization;
+	// Pack LumaRange into Params1.xy (the same slots chroma uses for AlphaThreshold).
+	Proxy->Params1 = FVector4f(LumaRange.X, LumaRange.Y, 0.0f, 0.0f);
+	Proxy->bDisplaySpaceLuminance = bDisplaySpaceLuminance;
+
+	return Proxy;
+}
+
diff --git a/Engine/Plugins/Compositing/Composite/Source/Composite/Private/Tests/CompositePassColorKeyerTest.cpp b/Engine/Plugins/Compositing/Composite/Source/Composite/Private/Tests/CompositePassColorKeyerTest.cpp
--- a/Engine/Plugins/Compositing/Composite/Source/Composite/Private/Tests/CompositePassColorKeyerTest.cpp
+++ b/Engine/Plugins/Compositing/Composite/Source/Composite/Private/Tests/CompositePassColorKeyerTest.cpp
@@ -102,7 +102,7 @@ bool FCompositeColorKeyerPassTest::RunTest(const FString& Parameter)
 					FClearValueBinding::Black, ETextureCreateFlags::None);
 				FRDGTextureRef InputTexture = GraphBuilder.CreateTexture(InputDesc, TEXT("ColorKeyerTest.Input"));
 
-				// ColorKeyer writes via UAV: output needs UAV + RenderTargetable flags.
+				// Output needs RenderTargetable + UAV flags.
 				const FRDGTextureDesc OutputDesc = FRDGTextureDesc::Create2D(
 					GTestViewSize, PF_A32B32G32R32F,
 					FClearValueBinding::Black,
@@ -186,4 +186,169 @@ bool FCompositeColorKeyerPassTest::RunTest(const FString& Parameter)
 
 #endif // WITH_DEV_AUTOMATION_TESTS
 
+// -------------------------------------------------------------------------------------------------
+
+#if WITH_DEV_AUTOMATION_TESTS
+
+namespace LumaKeyer  // unique name prevents symbol collision in unity builds
+{
+
+struct FLumaKeyerScenario
+{
+	const TCHAR* Label;
+	FLinearColor  Input;
+	FVector2f     LumaRange;
+	bool          bInvertAlpha;
+	bool          bDisplaySpaceLuminance;
+	FLinearColor  Expected;  // premultiplied output
+};
+
+static const FLumaKeyerScenario GScenarios[] =
+{
+	// Pure black: luminance=0 → alpha=0 → fully transparent.
+	{ TEXT("BlackFullRange"),         { 0.0f, 0.0f, 0.0f, 1.0f }, { 0.0f, 1.0f }, false, false, { 0.0f,    0.0f,    0.0f,    0.0f    } },
+	// Pure white: luminance=1 → alpha=1 → fully opaque. Premultiplied RGB unchanged.
+	{ TEXT("WhiteFullRange"),         { 1.0f, 1.0f, 1.0f, 1.0f }, { 0.0f, 1.0f }, false, false, { 1.0f,    1.0f,    1.0f,    1.0f    } },
+	// Mid-grey (linear): luminance=0.5 → alpha=0.5 → premultiplied RGB halved.
+	{ TEXT("GreyMidpoint"),           { 0.5f, 0.5f, 0.5f, 1.0f }, { 0.0f, 1.0f }, false, false, { 0.25f,   0.25f,   0.25f,   0.5f    } },
+	// Dark grey below [0.4,0.6] range → alpha clamps to 0.
+	{ TEXT("DarkBelowThreshold"),     { 0.2f, 0.2f, 0.2f, 1.0f }, { 0.4f, 0.6f }, false, false, { 0.0f,    0.0f,    0.0f,    0.0f    } },
+	// Bright grey above [0.4,0.6] range → alpha clamps to 1. Premultiplied RGB unchanged.
+	{ TEXT("BrightAboveThreshold"),   { 0.8f, 0.8f, 0.8f, 1.0f }, { 0.4f, 0.6f }, false, false, { 0.8f,    0.8f,    0.8f,    1.0f    } },
+	// White + inverted alpha: raw alpha=1 → inverted to 0 → premultiplied output is zero.
+	{ TEXT("WhiteInverted"),          { 1.0f, 1.0f, 1.0f, 1.0f }, { 0.0f, 1.0f }, true,  false, { 0.0f,    0.0f,    0.0f,    0.0f    } },
+	// Mid-grey with display-space luminance: LinearToSrgb(0.5)≈0.73536 → luma≈0.73536 → alpha≈0.73536.
+	// Premultiplied: RGB = 0.5 * 0.73536 ≈ 0.36768. Verifies the LinearToSrgb path in the shader.
+	{ TEXT("GreyMidpointDisplaySpace"),{ 0.5f, 0.5f, 0.5f, 1.0f }, { 0.0f, 1.0f }, false, true,  { 0.36768f, 0.36768f, 0.36768f, 0.73536f } },
+};
+
+}  // namespace LumaKeyer
+
+IMPLEMENT_SIMPLE_AUTOMATION_TEST(FCompositeLumaKeyerPassTest, "Composite.UnitTests.LumaKeyerPass",
+	EAutomationTestFlags_ApplicationContextMask | EAutomationTestFlags::EngineFilter | EAutomationTestFlags::NonNullRHI)
+
+bool FCompositeLumaKeyerPassTest::RunTest(const FString& Parameter)
+{
+	using namespace UE::Composite::Private;
+	using namespace UE::CompositeCore;
+	using namespace LumaKeyer;
+
+	const int32 NumScenarios = UE_ARRAY_COUNT(GScenarios);
+
+	TArray<FLinearColor> ResultRDG;
+	ResultRDG.SetNumZeroed(NumScenarios);
+
+	TArray<FCompositePassColorKeyerProxy> Proxies;
+	Proxies.Reserve(NumScenarios);
+	for (int32 i = 0; i < NumScenarios; ++i)
+	{
+		const FLumaKeyerScenario& S = GScenarios[i];
+		FCompositePassColorKeyerProxy Proxy(GetDefaultInputDeclArray());
+		Proxy.KeyerMode              = ECompositeColorKeyerMode::Luma;
+		// Pack LumaRange into Params1.xy — same packing GetProxy() uses for luma mode.
+		Proxy.Params1                = FVector4f(S.LumaRange.X, S.LumaRange.Y, 0.0f, 0.0f);
+		Proxy.Visualization          = ECompositeColorKeyerVisualization::Key;
+		Proxy.bInvertAlpha           = S.bInvertAlpha;
+		Proxy.bDisplaySpaceLuminance = S.bDisplaySpaceLuminance;
+		Proxies.Emplace(MoveTemp(Proxy));
+	}
+
+	FCompositeTestView TestView;
+	DispatchAndWait(TestView.Get(),
+		[&ResultRDG, NumScenarios, Proxies = MoveTemp(Proxies)](FRDGBuilder& GraphBuilder, const FSceneView& View) mutable
+		{
+			const FIntRect PixelRect(0, 0, GTestViewSize.X, GTestViewSize.Y);
+
+			for (int32 i = 0; i < NumScenarios; ++i)
+			{
+				const FLumaKeyerScenario& S = GScenarios[i];
+
+				const FRDGTextureDesc InputDesc = FRDGTextureDesc::Create2D(
+					GTestViewSize, PF_A32B32G32R32F,
+					FClearValueBinding::Black, ETextureCreateFlags::None);
+				FRDGTextureRef InputTexture = GraphBuilder.CreateTexture(InputDesc, TEXT("LumaKeyerTest.Input"));
+
+				const FRDGTextureDesc OutputDesc = FRDGTextureDesc::Create2D(
+					GTestViewSize, PF_A32B32G32R32F,
+					FClearValueBinding::Black,
+					ETextureCreateFlags::RenderTargetable | ETextureCreateFlags::UAV);
+				FRDGTextureRef OutputTexture = GraphBuilder.CreateTexture(OutputDesc, TEXT("LumaKeyerTest.Output"));
+
+				{
+					FColorKeyerTestUploadParameters* Params = GraphBuilder.AllocParameters<FColorKeyerTestUploadParameters>();
+					Params->InputTexture = InputTexture;
+
+					const FLinearColor PixelValue = S.Input;
+					GraphBuilder.AddPass(
+						RDG_EVENT_NAME("LumaKeyerTest.Upload[%d]", i), Params,
+						ERDGPassFlags::Readback,
+						[PixelValue, Params](FRHICommandListImmediate& RHICmdList)
+						{
+							uint32 Stride = 0;
+							void* Data = RHICmdList.LockTexture2D(Params->InputTexture->GetRHI(), 0, RLM_WriteOnly, Stride, false);
+							for (int32 Row = 0; Row < GTestViewSize.Y; ++Row)
+							{
+								FLinearColor* RowPtr = reinterpret_cast<FLinearColor*>(static_cast<uint8*>(Data) + Row * Stride);
+								for (int32 Col = 0; Col < GTestViewSize.X; ++Col)
+								{
+									RowPtr[Col] = PixelValue;
+								}
+							}
+							RHICmdList.UnlockTexture2D(Params->InputTexture->GetRHI(), 0, false);
+						});
+				}
+
+				FPassInputArray Inputs;
+				{
+					FPassInput& In = Inputs.GetArray().AddDefaulted_GetRef();
+					In.Texture         = FScreenPassTexture{ InputTexture, PixelRect };
+					In.Metadata.Filter = SF_Point;
+				}
+				Inputs.OverrideOutput = FScreenPassRenderTarget{
+					FScreenPassTexture{ OutputTexture, PixelRect }, ERenderTargetLoadAction::ENoAction };
+
+				FPassContext PassContext;
+				PassContext.OutputViewRect = PixelRect;
+
+				Proxies[i].Add(GraphBuilder, View, Inputs, PassContext);
+
+				{
+					FColorKeyerTestDownloadParameters* Params = GraphBuilder.AllocParameters<FColorKeyerTestDownloadParameters>();
+					Params->OutputTexture = OutputTexture;
+
+					FLinearColor* ResultPtr = &ResultRDG[i];
+					GraphBuilder.AddPass(
+						RDG_EVENT_NAME("LumaKeyerTest.Download[%d]", i), Params,
+						ERDGPassFlags::Readback,
+						[ResultPtr, Params](FRHICommandListImmediate& RHICmdList)
+						{
+							uint32 Stride = 0;
+							const void* Data = RHICmdList.LockTexture2D(Params->OutputTexture->GetRHI(), 0, RLM_ReadOnly, Stride, false);
+							const FLinearColor* RowPtr = reinterpret_cast<const FLinearColor*>(
+								static_cast<const uint8*>(Data) + (GTestViewSize.Y / 2) * Stride);
+							FMemory::Memcpy(ResultPtr, &RowPtr[GTestViewSize.X / 2], sizeof(FLinearColor));
+							RHICmdList.UnlockTexture2D(Params->OutputTexture->GetRHI(), 0, false);
+						});
+				}
+			}
+		});
+
+	constexpr float Tolerance = 1e-4f;
+	for (int32 i = 0; i < NumScenarios; ++i)
+	{
+		const FLumaKeyerScenario& S = GScenarios[i];
+		const FLinearColor& Expected = S.Expected;
+		const FLinearColor& Actual   = ResultRDG[i];
+
+		UTEST_EQUAL_TOLERANCE(S.Label, Actual.R, Expected.R, Tolerance);
+		UTEST_EQUAL_TOLERANCE(S.Label, Actual.G, Expected.G, Tolerance);
+		UTEST_EQUAL_TOLERANCE(S.Label, Actual.B, Expected.B, Tolerance);
+		UTEST_EQUAL_TOLERANCE(S.Label, Actual.A, Expected.A, Tolerance);
+	}
+
+	return true;
+}
+
+#endif // WITH_DEV_AUTOMATION_TESTS
+
 } // namespace UE::Composite::Tests
diff --git a/Engine/Plugins/Compositing/Composite/Source/Composite/Public/Passes/CompositePassColorKeyer.h b/Engine/Plugins/Compositing/Composite/Source/Composite/Public/Passes/CompositePassColorKeyer.h
--- a/Engine/Plugins/Compositing/Composite/Source/Composite/Public/Passes/CompositePassColorKeyer.h
+++ b/Engine/Plugins/Compositing/Composite/Source/Composite/Public/Passes/CompositePassColorKeyer.h
@@ -26,6 +26,14 @@ enum class ECompositeColorKeyerScreenType : uint8
 	Blue = 2
 };
 
+/** Keying mode: chroma key (color-based) or luma key (luminance-based). */
+UENUM(BlueprintType)
+enum class ECompositeColorKeyerMode : uint8
+{
+	Chroma = 0 UMETA(DisplayName = "Chroma Key"),
+	Luma   = 1 UMETA(DisplayName = "Luma Key"),
+};
+
 /** Denoising method applied before keying. */
 UENUM(BlueprintType)
 enum class ECompositeDenoiseMethod : uint8
@@ -59,7 +67,7 @@ public:
 	UE_API virtual FCompositeCorePassProxy* GetProxy(const UE::CompositeCore::FPassInputDecl& InputDecl, FCompositeTraversalContext& InContext, FSceneRenderingBulkObjectAllocator& InFrameAllocator) const override;
 
 	UE_API virtual void PostLoad() override;
-	
+
 #if WITH_EDITOR
 	UE_API virtual void PreEditChange(FProperty* PropertyThatWillChange) override;
 	UE_API virtual void PostEditChangeProperty(FPropertyChangedEvent& PropertyChangedEvent) override;
@@ -147,6 +155,45 @@ private:
 	friend class FCompositePassColorKeyerCustomization;
 };
 
+/**
+ * Luma keyer: derives alpha from pixel luminance.
+*/
+UCLASS(MinimalAPI, BlueprintType, Blueprintable, EditInlineNew, CollapseCategories, meta=(DisplayName="Luma Keyer"))
+class UCompositePassLumaKeyer : public UCompositePassBase
+{
+	GENERATED_BODY()
+
+public:
+	UE_API UCompositePassLumaKeyer(const FObjectInitializer& ObjectInitializer);
+	UE_API ~UCompositePassLumaKeyer();
+
+	UE_API virtual FCompositeCorePassProxy* GetProxy(const UE::CompositeCore::FPassInputDecl& InputDecl, FCompositeTraversalContext& InContext, FSceneRenderingBulkObjectAllocator& InFrameAllocator) const override;
+
+public:
+	/** Luminance range [min, max] that maps to alpha [0, 1]. Pixels with luminance below min are fully transparent; above max are fully opaque. */
+	UPROPERTY(EditAnywhere, BlueprintReadWrite, Interp, Category = "Composite", meta = (ClampMin = 0.0f, ClampMax = 1.0f, UIMin = 0.0f, UIMax = 1.0f))
+	FVector2f LumaRange;
+
+	/** Vizualize the alpha key or fill. */
+	UPROPERTY(EditAnywhere, BlueprintReadWrite, Category = "Composite")
+	ECompositeColorKeyerVisualization Visualization;
+
+	/** Invert the alpha key. */
+	UPROPERTY(EditAnywhere, BlueprintReadWrite, Category = "Composite", AdvancedDisplay)
+	bool bInvertAlpha;
+
+	/**
+	 * When enabled, luminance is computed after converting the pixel to display space (sRGB encoding) rather than in linear light.
+	 * This matches the luma-key behavior of other tools, where LumaRange values correspond to what you see on screen.
+	 */
+	UPROPERTY(EditAnywhere, BlueprintReadWrite, Category = "Composite", AdvancedDisplay)
+	bool bDisplaySpaceLuminance;
+
+	/** Denoising method applied before the keyer. */
+	UPROPERTY(EditAnywhere, BlueprintReadWrite, Category = "Composite", AdvancedDisplay)
+	ECompositeDenoiseMethod DenoiseMethod;
+};
+
 namespace UE
 {
 	namespace Composite
@@ -164,13 +211,15 @@ namespace UE
 				/** Adds the color keyer RDG pass and returns a premultiplied keyed RGBA output. */
 				FPassTexture Add(FRDGBuilder& GraphBuilder, const FSceneView& InView, const FPassInputArray& Inputs, const FPassContext& PassContext) const override;
 
+				/** Whether this proxy uses chroma or luma keying. */
+				ECompositeColorKeyerMode          KeyerMode              = ECompositeColorKeyerMode::Chroma;
 				/** Background screen color used to drive the key algorithm. */
 				ECompositeColorKeyerScreenType    ScreenType             = ECompositeColorKeyerScreenType::Green;
 				/** Static background reference color; used when no clean plate is provided. */
 				FLinearColor                      KeyColor;
 				/** Shader parameter vector: (RedWeight, GreenWeight, BlueWeight, DespillStrength). */
 				FVector4f                         Params0                = FVector4f(0.0f, 0.0f, 0.0f, 0.0f);
-				/** Shader parameter vector: (AlphaMin, AlphaMax, DevignetteStrength, unused). */
+				/** Shader parameter vector: chroma=(AlphaMin, AlphaMax, DevignetteStrength, unused), luma=(LumaMin, LumaMax, 0, 0). */
 				FVector4f                         Params1                = FVector4f(0.0f, 0.0f, 0.0f, 0.0f);
 				/** Which output channel to display: final composite key, fill layer, or raw alpha matte. */
 				ECompositeColorKeyerVisualization Visualization          = ECompositeColorKeyerVisualization::Key;
@@ -182,6 +231,8 @@ namespace UE
 				bool                              bDespillOnly           = false;
 				/** When true, the output alpha is inverted (1 − alpha). */
 				bool                              bInvertAlpha           = false;
+				/** When true, luminance is computed in display (sRGB encoding) space to match other tools. */
+				bool                              bDisplaySpaceLuminance = false;
 				/** Weak pointer to the clean plate texture used as the per-pixel background reference. Null when KeyColor is used instead. */
 				TWeakObjectPtr<UTexture>          CleanPlateWeakPtr;
 			};
