Composure: Fix localization conflict warning.

#Synced-CL 52610045

diff --git a/Engine/Plugins/Compositing/Composite/Source/CompositeEditor/Private/Customizations/CompositeLayerPlateCustomization.cpp b/Engine/Plugins/Compositing/Composite/Source/CompositeEditor/Private/Customizations/CompositeLayerPlateCustomization.cpp
--- a/Engine/Plugins/Compositing/Composite/Source/CompositeEditor/Private/Customizations/CompositeLayerPlateCustomization.cpp
+++ b/Engine/Plugins/Compositing/Composite/Source/CompositeEditor/Private/Customizations/CompositeLayerPlateCustomization.cpp
@@ -354,7 +354,7 @@ void FCompositeLayerPlateCustomization::OnCompositeMeshActorToggled(TSoftObjectP
 	UCompositeLayerPlate* Plate = Plates[0].Get();
 	TArray<UPrimitiveComponent*> CompositeMeshPrimitives = Plate->GetIncludedPrimitivesForActor(InActor);
 
-	FScopedTransaction Transaction(LOCTEXT("ToggleCompositeMeshPrimitives", "Toggle Composite Mesh Actor"));
+	FScopedTransaction Transaction(LOCTEXT("ToggleCompositeMeshActor", "Toggle Composite Mesh Actor"));
 	Plate->Modify();
 	
 	// Toggling the whole actor means removing all specifically included primitives so that all valid primitives in the actor are used
