Composure: Expose option to disable display transform emulation for FXAA/SMAA, since it can negatively impact alpha.

#Synced-CL 52632071

diff --git a/Engine/Plugins/Compositing/Composite/Source/Composite/Private/Passes/CompositePassFXAA.cpp b/Engine/Plugins/Compositing/Composite/Source/Composite/Private/Passes/CompositePassFXAA.cpp
--- a/Engine/Plugins/Compositing/Composite/Source/Composite/Private/Passes/CompositePassFXAA.cpp
+++ b/Engine/Plugins/Compositing/Composite/Source/Composite/Private/Passes/CompositePassFXAA.cpp
@@ -18,6 +18,7 @@ FCompositeCorePassProxy* UCompositePassFXAA::GetProxy(const UE::CompositeCore::F
 
 	FFXAAPassProxy* Proxy = InFrameAllocator.Create<FFXAAPassProxy>(UE::CompositeCore::FPassInputDeclArray{ InputDecl });
 	Proxy->QualityOverride = Quality;
+	Proxy->bDisplayTransform = bDisplayTransform;
 
 	return Proxy;
 }
diff --git a/Engine/Plugins/Compositing/Composite/Source/Composite/Private/Passes/CompositePassSMAA.cpp b/Engine/Plugins/Compositing/Composite/Source/Composite/Private/Passes/CompositePassSMAA.cpp
--- a/Engine/Plugins/Compositing/Composite/Source/Composite/Private/Passes/CompositePassSMAA.cpp
+++ b/Engine/Plugins/Compositing/Composite/Source/Composite/Private/Passes/CompositePassSMAA.cpp
@@ -33,8 +33,9 @@ namespace UE
 				FScreenPassTexture Input = Inputs[0].Texture;
 				FResourceMetadata Metadata = Inputs[0].Metadata;
 				const bool bLinearSourceColors = (Metadata.Encoding == EEncoding::Linear);
+				const bool bApplyDisplayTransform = bLinearSourceColors && bDisplayTransform;
 
-				if (bLinearSourceColors)
+				if (bApplyDisplayTransform)
 				{
 					// We tonemap & encode the result so that SMAA can operate on perceptual colors
 					constexpr bool bIsForward = true;
@@ -50,14 +51,14 @@ namespace UE
 				SMAAInputs.SceneColor = Input;
 				SMAAInputs.SceneColorBeforeTonemap = Inputs[0].Texture;
 				SMAAInputs.Quality = static_cast<ESMAAQuality>(FMath::Clamp(Quality, 0, static_cast<int32>(ESMAAQuality::MAX) - 1));
-				if (!bLinearSourceColors)
+				if (!bApplyDisplayTransform)
 				{
 					SMAAInputs.OverrideOutput = Inputs.OverrideOutput;
 				}
 
 				FScreenPassTexture Output = AddSMAAPasses(GraphBuilder, InView, SMAAInputs);
 
-				if (bLinearSourceColors)
+				if (bApplyDisplayTransform)
 				{
 					// We decode and invert the tonemapping to obtain linear colors again.
 					constexpr bool bIsForward = false;
@@ -69,6 +70,9 @@ namespace UE
 
 			/** Optional r.SMAA.Quality setting override.  */
 			TOptional<int32> QualityOverride = {};
+
+			/** When true, applies a display transform (tonemap + gamma) before AA and inverts it after. */
+			bool bDisplayTransform = true;
 		};
 	}
 }
@@ -94,6 +98,7 @@ FCompositeCorePassProxy* UCompositePassSMAA::GetProxy(const UE::CompositeCore::F
 
 	FSMAAPassProxy* Proxy = InFrameAllocator.Create<FSMAAPassProxy>(UE::CompositeCore::FPassInputDeclArray{ InputDecl });
 	Proxy->QualityOverride = Quality;
+	Proxy->bDisplayTransform = bDisplayTransform;
 
 	return Proxy;
 }
diff --git a/Engine/Plugins/Compositing/Composite/Source/Composite/Public/Passes/CompositePassFXAA.h b/Engine/Plugins/Compositing/Composite/Source/Composite/Public/Passes/CompositePassFXAA.h
--- a/Engine/Plugins/Compositing/Composite/Source/Composite/Public/Passes/CompositePassFXAA.h
+++ b/Engine/Plugins/Compositing/Composite/Source/Composite/Public/Passes/CompositePassFXAA.h
@@ -28,6 +28,10 @@ public:
 	/** FXAA quality control, matching the "r.fxaa.quality" console variable values. */
 	UPROPERTY(EditAnywhere, BlueprintReadWrite, Interp, Category = "Composite", meta=(ClampMin="0", ClampMax="5", UIMin = "0", UIMax = "5"))
 	int32 Quality;
+
+	/** Apply a display transform (tonemap + gamma encode) before AA and invert it after. Improves perceptual edge detection on linear content but may darken edges at alpha boundaries. */
+	UPROPERTY(EditAnywhere, BlueprintReadWrite, Category = "Composite", AdvancedDisplay)
+	bool bDisplayTransform = true;
 };
 
 #undef UE_API
diff --git a/Engine/Plugins/Compositing/Composite/Source/Composite/Public/Passes/CompositePassSMAA.h b/Engine/Plugins/Compositing/Composite/Source/Composite/Public/Passes/CompositePassSMAA.h
--- a/Engine/Plugins/Compositing/Composite/Source/Composite/Public/Passes/CompositePassSMAA.h
+++ b/Engine/Plugins/Compositing/Composite/Source/Composite/Public/Passes/CompositePassSMAA.h
@@ -28,6 +28,10 @@ public:
 	/** SMAA quality control, matching the "r.SMAA.quality" console variable values. */
 	UPROPERTY(EditAnywhere, BlueprintReadWrite, Interp, Category = "Composite", meta=(ClampMin="0", ClampMax="5", UIMin = "0", UIMax = "5"))
 	int32 Quality;
+
+	/** Apply a display transform (tonemap + gamma encode) before AA and invert it after. Improves perceptual edge detection on linear content but may darken edges at alpha boundaries. */
+	UPROPERTY(EditAnywhere, BlueprintReadWrite, Category = "Composite", AdvancedDisplay)
+	bool bDisplayTransform = true;
 };
 
 #undef UE_API
diff --git a/Engine/Plugins/Compositing/CompositeCore/Source/CompositeCore/Private/Passes/CompositeCorePassFXAAProxy.cpp b/Engine/Plugins/Compositing/CompositeCore/Source/CompositeCore/Private/Passes/CompositeCorePassFXAAProxy.cpp
--- a/Engine/Plugins/Compositing/CompositeCore/Source/CompositeCore/Private/Passes/CompositeCorePassFXAAProxy.cpp
+++ b/Engine/Plugins/Compositing/CompositeCore/Source/CompositeCore/Private/Passes/CompositeCorePassFXAAProxy.cpp
@@ -63,7 +63,6 @@ namespace UE
 				Parameters->bIsForward = bIsForward;
 				Parameters->Gamma = InGamma;
 				Parameters->InvGamma = FMath::SafeDivide(1.0f, InGamma);
-				Parameters->bIsForward = bIsForward;
 				Parameters->RenderTargets[0] = Output.GetRenderTargetBinding();
 
 				TShaderMapRef<FCompositeCoreDisplayTransformShader> PixelShader(GlobalShaderMap);
@@ -89,8 +88,9 @@ namespace UE
 			FScreenPassTexture Input = Inputs[0].Texture;
 			FResourceMetadata Metadata = Inputs[0].Metadata;
 			const bool bLinearSourceColors = (Metadata.Encoding == EEncoding::Linear);
+			const bool bApplyDisplayTransform = bLinearSourceColors && bDisplayTransform;
 
-			if (bLinearSourceColors)
+			if (bApplyDisplayTransform)
 			{
 				// We tonemap & encode the result so that FXAA can operate on perceptual colors
 				constexpr bool bIsForward = true;
@@ -104,14 +104,14 @@ namespace UE
 			FFXAAInputs FXAAInputs{};
 			FXAAInputs.SceneColor = Input;
 			FXAAInputs.Quality = static_cast<EFXAAQuality>(FMath::Clamp(Quality, 0, static_cast<int32>(EFXAAQuality::MAX) - 1));
-			if (!bLinearSourceColors)
+			if (!bApplyDisplayTransform)
 			{
 				FXAAInputs.OverrideOutput = Inputs.OverrideOutput;
 			}
 
 			FScreenPassTexture Output = AddFXAAPass(GraphBuilder, InView, FXAAInputs);
 
-			if (bLinearSourceColors)
+			if (bApplyDisplayTransform)
 			{
 				// We decode and invert the tonemapping to obtain linear colors again.
 				constexpr bool bIsForward = false;
@@ -122,4 +122,3 @@ namespace UE
 		}
 	}
 }
-
diff --git a/Engine/Plugins/Compositing/CompositeCore/Source/CompositeCore/Public/Passes/CompositeCorePassFXAAProxy.h b/Engine/Plugins/Compositing/CompositeCore/Source/CompositeCore/Public/Passes/CompositeCorePassFXAAProxy.h
--- a/Engine/Plugins/Compositing/CompositeCore/Source/CompositeCore/Public/Passes/CompositeCorePassFXAAProxy.h
+++ b/Engine/Plugins/Compositing/CompositeCore/Source/CompositeCore/Public/Passes/CompositeCorePassFXAAProxy.h
@@ -27,6 +27,9 @@ namespace UE
 
 			/** Optional r.FXAA.Quality setting override.  */
 			TOptional<int32> QualityOverride = {};
+
+			/** When true, applies a display transform (tonemap + gamma) before AA and inverts it after. Improves edge detection on linear content but may produce dark halos at alpha boundaries. */
+			bool bDisplayTransform = true;
 		};
 	}
 }
