Composure: Follow up to CL 52657025, change FXAA/SMAA display transform default to false to avoid alpha edge issues despite the anti-aliasing being less acurate in linear space.

#Synced-CL 52632071

diff --git a/Engine/Plugins/Compositing/Composite/Source/Composite/Public/Passes/CompositePassFXAA.h b/Engine/Plugins/Compositing/Composite/Source/Composite/Public/Passes/CompositePassFXAA.h
--- a/Engine/Plugins/Compositing/Composite/Source/Composite/Public/Passes/CompositePassFXAA.h
+++ b/Engine/Plugins/Compositing/Composite/Source/Composite/Public/Passes/CompositePassFXAA.h
@@ -31,7 +31,7 @@ public:
 
 	/** Apply a display transform (tonemap + gamma encode) before AA and invert it after. Improves perceptual edge detection on linear content but may darken edges at alpha boundaries. */
 	UPROPERTY(EditAnywhere, BlueprintReadWrite, Category = "Composite", AdvancedDisplay)
-	bool bDisplayTransform = true;
+	bool bDisplayTransform = false;
 };
 
 #undef UE_API
diff --git a/Engine/Plugins/Compositing/Composite/Source/Composite/Public/Passes/CompositePassSMAA.h b/Engine/Plugins/Compositing/Composite/Source/Composite/Public/Passes/CompositePassSMAA.h
--- a/Engine/Plugins/Compositing/Composite/Source/Composite/Public/Passes/CompositePassSMAA.h
+++ b/Engine/Plugins/Compositing/Composite/Source/Composite/Public/Passes/CompositePassSMAA.h
@@ -31,7 +31,7 @@ public:
 
 	/** Apply a display transform (tonemap + gamma encode) before AA and invert it after. Improves perceptual edge detection on linear content but may darken edges at alpha boundaries. */
 	UPROPERTY(EditAnywhere, BlueprintReadWrite, Category = "Composite", AdvancedDisplay)
-	bool bDisplayTransform = true;
+	bool bDisplayTransform = false;
 };
 
 #undef UE_API
