Composure: Add geometry masking mode to CompositePassMasking, so that users can easily apply masks with geometry and not only textures.

- Add ECompositeMaskingMode enum (Texture/Geometry) to switch between texture and geometry-based masks
- Generalize ACompositeActor scene capture API from UCompositeLayerBase to UCompositePassBase
- Implement geometry mode: renders selected mesh actors via custom render pass, uses alpha as mask
- Add FTickableGameObject for per-frame ShowOnlyComponents refresh on the masking pass

#Synced-CL 52632071

diff --git a/Engine/Plugins/Compositing/Composite/Source/Composite/Private/CompositeActor.cpp b/Engine/Plugins/Compositing/Composite/Source/Composite/Private/CompositeActor.cpp
--- a/Engine/Plugins/Compositing/Composite/Source/Composite/Private/CompositeActor.cpp
+++ b/Engine/Plugins/Compositing/Composite/Source/Composite/Private/CompositeActor.cpp
@@ -150,7 +150,7 @@ namespace UE
 			/* Gather render work (per frame) given the current composite actor properties. */
 			UE::CompositeCore::FRenderWork GatherFrameRenderWork(
 				const ACompositeActor& InCompositeActor,
-				const TMap<TWeakObjectPtr<const UCompositeLayerBase>, FSceneCaptureComponentArray>& InSceneCapturesPerLayer)
+				const TMap<TWeakObjectPtr<const UCompositePassBase>, FSceneCaptureComponentArray>& InSceneCapturesPerPass)
 			{
 				if (InCompositeActor.GetCompositeLayers().IsEmpty())
 				{
@@ -236,13 +236,13 @@ namespace UE
 				
 				RenderWork.AllowedViewModes = GetViewModeIndices(InCompositeActor.AllowedViewModes);
 
-				// Currently, each layer may use at most 2 scene captures, so we reserve the maximum case.
-				RenderWork.SceneCapturesUpdateQueue.Reserve(InSceneCapturesPerLayer.Num() * 2);
+				// Heuristic reservation; actual count per pass may vary.
+				RenderWork.SceneCapturesUpdateQueue.Reserve(InSceneCapturesPerPass.Num() * 2);
 
-				for (const auto& Pair : InSceneCapturesPerLayer)
+				for (const auto& Pair : InSceneCapturesPerPass)
 				{
-					const TWeakObjectPtr<const UCompositeLayerBase>& Layer = Pair.Key;
-					if (Layer.IsValid() && Layer->GetIsEnabled())
+					const TWeakObjectPtr<const UCompositePassBase>& Pass = Pair.Key;
+					if (Pass.IsValid() && Pass->GetIsEnabled())
 					{
 						for (const TObjectPtr<UCompositeSceneCapture2DComponent>& SceneCapture : Pair.Value.Components)
 						{
@@ -441,7 +441,7 @@ void ACompositeActor::Tick(float DeltaTime)
 
 	if (UCompositeWorldSubsystem* Subsystem = UWorld::GetSubsystem<UCompositeWorldSubsystem>(GetWorld()))
 	{
-		Subsystem->SetRenderWork(UE::Composite::Private::GatherFrameRenderWork(*this, SceneCapturesPerLayer));
+		Subsystem->SetRenderWork(UE::Composite::Private::GatherFrameRenderWork(*this, SceneCapturesPerPass));
 
 		UE::CompositeCore::FBuiltInRenderPassOptions Options;
 		if (bOverridesViewUserFlags)
@@ -967,9 +967,9 @@ void ACompositeActor::SetVPRole(FGameplayTag Value)
 }
 
 template<>
-UCompositeSceneCapture2DComponent* ACompositeActor::FindOrCreateSceneCapture(const UCompositeLayerBase* InLayer, int32 InIndex, FName InBaseName)
+UCompositeSceneCapture2DComponent* ACompositeActor::FindOrCreateSceneCapture(const UCompositePassBase* InPass, int32 InIndex, FName InBaseName)
 {
-	if (!ensureMsgf(InLayer, TEXT("Expected non-null composite layer.")))
+	if (!ensureMsgf(InPass, TEXT("Expected non-null composite pass.")))
 	{
 		return nullptr;
 	}
@@ -979,12 +979,12 @@ UCompositeSceneCapture2DComponent* ACompositeActor::FindOrCreateSceneCapture(con
 		return nullptr;
 	}
 
-	if (!SceneCapturesPerLayer.Contains(InLayer))
+	if (!SceneCapturesPerPass.Contains(InPass))
 	{
-		SceneCapturesPerLayer.Add(InLayer, FSceneCaptureComponentArray{});
+		SceneCapturesPerPass.Add(InPass, FSceneCaptureComponentArray{});
 	}
 
-	TArray<TObjectPtr<UCompositeSceneCapture2DComponent>>& SceneCaptures = SceneCapturesPerLayer[InLayer].Components;
+	TArray<TObjectPtr<UCompositeSceneCapture2DComponent>>& SceneCaptures = SceneCapturesPerPass[InPass].Components;
 
 	if (SceneCaptures.IsValidIndex(InIndex) && IsValid(SceneCaptures[InIndex]))
 	{
@@ -1013,9 +1013,9 @@ UCompositeSceneCapture2DComponent* ACompositeActor::FindOrCreateSceneCapture(con
 }
 
 template<>
-UCompositeShadowReflectionCatcherComponent* ACompositeActor::FindOrCreateSceneCapture(const UCompositeLayerBase* InLayer, int32 InIndex, FName InBaseName)
+UCompositeShadowReflectionCatcherComponent* ACompositeActor::FindOrCreateSceneCapture(const UCompositePassBase* InPass, int32 InIndex, FName InBaseName)
 {
-	if (!ensureMsgf(InLayer, TEXT("Expected non-null composite layer.")))
+	if (!ensureMsgf(InPass, TEXT("Expected non-null composite pass.")))
 	{
 		return nullptr;
 	}
@@ -1025,12 +1025,12 @@ UCompositeShadowReflectionCatcherComponent* ACompositeActor::FindOrCreateSceneCa
 		return nullptr;
 	}
 
-	if (!SceneCapturesPerLayer.Contains(InLayer))
+	if (!SceneCapturesPerPass.Contains(InPass))
 	{
-		SceneCapturesPerLayer.Add(InLayer, {});
+		SceneCapturesPerPass.Add(InPass, {});
 	}
 
-	TArray<TObjectPtr<UCompositeSceneCapture2DComponent>>& SceneCaptureArray = SceneCapturesPerLayer[InLayer].Components;
+	TArray<TObjectPtr<UCompositeSceneCapture2DComponent>>& SceneCaptureArray = SceneCapturesPerPass[InPass].Components;
 
 	if (SceneCaptureArray.IsValidIndex(InIndex))
 	{
@@ -1058,9 +1058,9 @@ UCompositeShadowReflectionCatcherComponent* ACompositeActor::FindOrCreateSceneCa
 	return Component;
 }
 
-void ACompositeActor::DestroySceneCaptures(const UCompositeLayerBase* InLayer)
+void ACompositeActor::DestroySceneCaptures(const UCompositePassBase* InPass)
 {
-	if (!ensureMsgf(InLayer, TEXT("Expected non-null composite layer.")))
+	if (!ensureMsgf(InPass, TEXT("Expected non-null composite pass.")))
 	{
 		return;
 	}
@@ -1068,7 +1068,7 @@ void ACompositeActor::DestroySceneCaptures(const UCompositeLayerBase* InLayer)
 	bool bComponentRemoved = false;
 
 	FSceneCaptureComponentArray RemovedArray;
-	if (SceneCapturesPerLayer.RemoveAndCopyValue(InLayer, RemovedArray))
+	if (SceneCapturesPerPass.RemoveAndCopyValue(InPass, RemovedArray))
 	{
 		for (TObjectPtr<UCompositeSceneCapture2DComponent>& RemovedComponent : RemovedArray.Components)
 		{
diff --git a/Engine/Plugins/Compositing/Composite/Source/Composite/Private/Passes/CompositePassMasking.cpp b/Engine/Plugins/Compositing/Composite/Source/Composite/Private/Passes/CompositePassMasking.cpp
--- a/Engine/Plugins/Compositing/Composite/Source/Composite/Private/Passes/CompositePassMasking.cpp
+++ b/Engine/Plugins/Compositing/Composite/Source/Composite/Private/Passes/CompositePassMasking.cpp
@@ -3,7 +3,12 @@
 #include "Passes/CompositePassMasking.h"
 #include "Passes/CompositeCorePassProxy.h"
 
+#include "CompositeActor.h"
+#include "CompositeRenderTargetPool.h"
+#include "Components/CompositeSceneCapture2DComponent.h"
+#include "Components/PrimitiveComponent.h"
 #include "Engine/Texture.h"
+#include "Engine/TextureRenderTarget2D.h"
 #include "PixelShaderUtils.h"
 #include "RenderGraphBuilder.h"
 #include "RenderGraphUtils.h"
@@ -99,8 +104,11 @@ namespace UE::Composite::Private
 
 UCompositePassMasking::UCompositePassMasking(const FObjectInitializer& ObjectInitializer)
 	: Super(ObjectInitializer)
+	, MaskingMode{ ECompositeMaskingMode::Texture }
 	, MaskTexture{}
 	, MaskSource{ ECompositeColorChannel::Red }
+	, RenderTargetResolution{ FCompositeRenderTargetPool::DefaultSize }
+	, bVisibleInSceneCaptureOnly{ true }
 	, bUnpremult{ true }
 	, bInvert{ false }
 {
@@ -108,31 +116,317 @@ UCompositePassMasking::UCompositePassMasking(const FObjectInitializer& ObjectIni
 
 UCompositePassMasking::~UCompositePassMasking() = default;
 
+ETickableTickType UCompositePassMasking::GetTickableTickType() const
+{
+	return HasAnyFlags(RF_ClassDefaultObject) ? ETickableTickType::Never : ETickableTickType::Conditional;
+}
+
+bool UCompositePassMasking::IsTickable() const
+{
+	if (HasAnyFlags(RF_BeginDestroyed) || MaskingMode != ECompositeMaskingMode::Geometry)
+	{
+		return false;
+	}
+
+	ACompositeActor* CompositeActor = GetTypedOuter<ACompositeActor>();
+	return IsValid(CompositeActor) && !CompositeActor->HasAnyFlags(RF_ClassDefaultObject);
+}
+
+UWorld* UCompositePassMasking::GetTickableGameObjectWorld() const
+{
+	return GetWorld();
+}
+
+TStatId UCompositePassMasking::GetStatId() const
+{
+	RETURN_QUICK_DECLARE_CYCLE_STAT(UCompositePassMasking, STATGROUP_Tickables);
+}
+
+void UCompositePassMasking::Tick(float DeltaTime)
+{
+	UCompositeSceneCapture2DComponent* SceneCaptureComponent = FindOrCreateSceneCaptureComponent();
+	if (!IsValid(SceneCaptureComponent))
+	{
+		return;
+	}
+
+	// Refresh ShowOnlyComponents every frame to handle primitive recreation.
+	SceneCaptureComponent->ShowOnlyComponents.Reset(MaskActors.Num());
+	for (const TSoftObjectPtr<AActor>& SoftActor : MaskActors)
+	{
+		if (const AActor* Actor = SoftActor.Get())
+		{
+			for (UActorComponent* Component : Actor->GetComponents())
+			{
+				if (UPrimitiveComponent* PrimitiveComponent = Cast<UPrimitiveComponent>(Component))
+				{
+					SceneCaptureComponent->ShowOnlyComponents.Add(PrimitiveComponent);
+				}
+			}
+		}
+	}
+}
+
+UCompositeSceneCapture2DComponent* UCompositePassMasking::FindOrCreateSceneCaptureComponent() const
+{
+	ACompositeActor* CompositeActor = GetTypedOuter<ACompositeActor>();
+	if (!IsValid(CompositeActor))
+	{
+		return nullptr;
+	}
+
+	UCompositeSceneCapture2DComponent* CaptureComponent = CompositeActor->FindOrCreateSceneCapture<UCompositeSceneCapture2DComponent>(this, 0, FName("MaskingGeometryCapture"));
+	if (IsValid(CaptureComponent) && !CaptureComponent->bRenderInMainRenderer)
+	{
+		// Configure as custom render pass inlined in the main renderer (shared scene setup, reduced overhead).
+		CaptureComponent->CaptureSource = ESceneCaptureSource::SCS_SceneColorHDR;
+		CaptureComponent->bRenderInMainRenderer = true;
+		CaptureComponent->PrimitiveRenderMode = ESceneCapturePrimitiveRenderMode::PRM_UseShowOnlyList;
+	}
+
+	return CaptureComponent;
+}
+
+void UCompositePassMasking::UpdatePrimitiveVisibilityState(bool bInVisibleInSceneCaptureOnly) const
+{
+	for (const TSoftObjectPtr<AActor>& SoftActor : MaskActors)
+	{
+		if (const AActor* Actor = SoftActor.Get())
+		{
+			for (UActorComponent* Component : Actor->GetComponents())
+			{
+				if (UPrimitiveComponent* PrimitiveComponent = Cast<UPrimitiveComponent>(Component))
+				{
+					if (PrimitiveComponent->bVisibleInSceneCaptureOnly != bInVisibleInSceneCaptureOnly)
+					{
+						PrimitiveComponent->Modify();
+						PrimitiveComponent->SetVisibleInSceneCaptureOnly(bInVisibleInSceneCaptureOnly);
+					}
+				}
+			}
+		}
+	}
+}
+
+void UCompositePassMasking::RestorePrimitiveVisibilityState() const
+{
+	// Restores to CDO default rather than caching original per-component state.
+	// Primitives may be recreated or modified externally at any time, making cached state unreliable.
+	UPrimitiveComponent* CDO = UPrimitiveComponent::StaticClass()->GetDefaultObject<UPrimitiveComponent>();
+	UpdatePrimitiveVisibilityState(CDO->bVisibleInSceneCaptureOnly);
+}
+
+void UCompositePassMasking::BeginDestroy()
+{
+	if (MaskingMode == ECompositeMaskingMode::Geometry)
+	{
+		if (ACompositeActor* CompositeActor = GetTypedOuter<ACompositeActor>())
+		{
+			CompositeActor->DestroySceneCaptures(this);
+		}
+
+		RestorePrimitiveVisibilityState();
+	}
+
+	Super::BeginDestroy();
+}
+
+void UCompositePassMasking::SetMaskingMode(ECompositeMaskingMode InMaskingMode)
+{
+	if (MaskingMode == InMaskingMode)
+	{
+		return;
+	}
+
+	// Clean up previous geometry mode state.
+	if (MaskingMode == ECompositeMaskingMode::Geometry)
+	{
+		if (ACompositeActor* CompositeActor = GetTypedOuter<ACompositeActor>())
+		{
+			CompositeActor->DestroySceneCaptures(this);
+		}
+
+		RestorePrimitiveVisibilityState();
+	}
+
+	MaskingMode = InMaskingMode;
+
+	// Initialize new geometry mode state.
+	if (MaskingMode == ECompositeMaskingMode::Geometry)
+	{
+		UpdatePrimitiveVisibilityState(bVisibleInSceneCaptureOnly);
+	}
+}
+
+const TArray<TSoftObjectPtr<AActor>>& UCompositePassMasking::GetMaskActors() const
+{
+	return MaskActors;
+}
+
+void UCompositePassMasking::SetMaskActors(TArray<TSoftObjectPtr<AActor>> InMaskActors)
+{
+	if (MaskingMode == ECompositeMaskingMode::Geometry)
+	{
+		RestorePrimitiveVisibilityState();
+	}
+
+	MaskActors = MoveTemp(InMaskActors);
+
+	if (MaskingMode == ECompositeMaskingMode::Geometry)
+	{
+		UpdatePrimitiveVisibilityState(bVisibleInSceneCaptureOnly);
+	}
+}
+
+bool UCompositePassMasking::IsVisibleInSceneCaptureOnly() const
+{
+	return bVisibleInSceneCaptureOnly;
+}
+
+void UCompositePassMasking::SetVisibleInSceneCaptureOnly(bool bInVisible)
+{
+	bVisibleInSceneCaptureOnly = bInVisible;
+
+	if (MaskingMode == ECompositeMaskingMode::Geometry)
+	{
+		UpdatePrimitiveVisibilityState(bVisibleInSceneCaptureOnly);
+	}
+}
+
+void UCompositePassMasking::SetRenderTargetResolution(FIntPoint InRenderTargetResolution)
+{
+	RenderTargetResolution.X = FMath::Max(InRenderTargetResolution.X, 1);
+	RenderTargetResolution.Y = FMath::Max(InRenderTargetResolution.Y, 1);
+}
+
+#if WITH_EDITOR
+void UCompositePassMasking::PreEditChange(FProperty* PropertyAboutToChange)
+{
+	Super::PreEditChange(PropertyAboutToChange);
+
+	if (!PropertyAboutToChange)
+	{
+		return;
+	}
+
+	const FName PropertyName = PropertyAboutToChange->GetFName();
+
+	if (PropertyName == GET_MEMBER_NAME_CHECKED(ThisClass, MaskingMode))
+	{
+		PreEditMaskingMode = MaskingMode;
+	}
+	else if (PropertyName == GET_MEMBER_NAME_CHECKED(ThisClass, MaskActors) && MaskingMode == ECompositeMaskingMode::Geometry)
+	{
+		RestorePrimitiveVisibilityState();
+	}
+}
+
+void UCompositePassMasking::PostEditChangeProperty(FPropertyChangedEvent& PropertyChangedEvent)
+{
+	Super::PostEditChangeProperty(PropertyChangedEvent);
+
+	const FName PropertyName = PropertyChangedEvent.GetMemberPropertyName();
+
+	if (PropertyName == GET_MEMBER_NAME_CHECKED(ThisClass, MaskingMode))
+	{
+		/**
+		 * The editor has already written the new value into MaskingMode.
+		 * Restore the old value so SetMaskingMode sees the actual transition.
+		 */
+		ECompositeMaskingMode NewMode = MaskingMode;
+		MaskingMode = PreEditMaskingMode;
+		SetMaskingMode(NewMode);
+	}
+	else if (PropertyName == GET_MEMBER_NAME_CHECKED(ThisClass, MaskActors))
+	{
+		SetMaskActors(MaskActors);
+	}
+	else if (PropertyName == GET_MEMBER_NAME_CHECKED(ThisClass, bVisibleInSceneCaptureOnly))
+	{
+		SetVisibleInSceneCaptureOnly(bVisibleInSceneCaptureOnly);
+	}
+	else if (PropertyName == GET_MEMBER_NAME_CHECKED(ThisClass, RenderTargetResolution))
+	{
+		SetRenderTargetResolution(RenderTargetResolution);
+	}
+}
+#endif
+
 bool UCompositePassMasking::GetIsActive() const
 {
-	return GetIsEnabled() && IsValid(MaskTexture);
+	if (!GetIsEnabled())
+	{
+		return false;
+	}
+
+	switch (MaskingMode)
+	{
+	case ECompositeMaskingMode::Texture:
+		return IsValid(MaskTexture);
+
+	case ECompositeMaskingMode::Geometry:
+	{
+		if (!IsValid(GetTypedOuter<ACompositeActor>()))
+		{
+			return false;
+		}
+		for (const TSoftObjectPtr<AActor>& SoftActor : MaskActors)
+		{
+			if (SoftActor.Get())
+			{
+				return true;
+			}
+		}
+		return false;
+	}
+
+	default:
+		return false;
+	}
 }
 
 FCompositeCorePassProxy* UCompositePassMasking::GetProxy(const UE::CompositeCore::FPassInputDecl& InputDecl, FCompositeTraversalContext& InContext, FSceneRenderingBulkObjectAllocator& InFrameAllocator) const
 {
 	using namespace UE::Composite::Private;
+	using namespace UE::CompositeCore;
 
-	if (!ensureMsgf(IsValid(MaskTexture), TEXT("Texture is expected to be valid for active pass.")))
+	FPassInputDecl MaskInputDecl;
+
+	if (MaskingMode == ECompositeMaskingMode::Texture)
 	{
-		return nullptr;
+		if (!ensureMsgf(IsValid(MaskTexture), TEXT("Texture is expected to be valid for active pass.")))
+		{
+			return nullptr;
+		}
+
+		FResourceMetadata MaskMetadata = { .DebugName = TEXT("MaskTexture") };
+		MaskMetadata.Filter = (MaskTexture->Filter == TF_Nearest) ? SF_Point : SF_Bilinear;
+
+		ResourceId TextureId = InContext.FindOrCreateExternalTexture(MaskTexture, MaskMetadata);
+		MaskInputDecl.Set<FPassExternalResourceDesc>({ TextureId });
 	}
+	else // Geometry
+	{
+		UCompositeSceneCapture2DComponent* CaptureComponent = FindOrCreateSceneCaptureComponent();
+		if (!ensureMsgf(IsValid(CaptureComponent), TEXT("Scene capture component is expected to be valid for active geometry mask pass.")))
+		{
+			return nullptr;
+		}
 
-	UE::CompositeCore::FResourceMetadata MaskMetadata = { .DebugName = TEXT("MaskTexture") };
-	MaskMetadata.Filter = (MaskTexture->Filter == TF_Nearest) ? SF_Point : SF_Bilinear;
+		FCompositeRenderTargetPool::Get().ConditionalAcquireTarget(CaptureComponent, CaptureComponent->TextureTarget, RenderTargetResolution);
 
-	ResourceId TextureId = InContext.FindOrCreateExternalTexture(MaskTexture, MaskMetadata);
+		FResourceMetadata MaskMetadata = {};
+		MaskMetadata.bInvertedAlpha = true;
+		MaskMetadata.DebugName = TEXT("MaskingGeometryCaptureTex");
 
-	FPassInputDecl MaskingTextureInput;
-	MaskingTextureInput.Set<FPassExternalResourceDesc>({ TextureId });
+		ResourceId TextureId = InContext.FindOrCreateExternalTexture(CaptureComponent->TextureTarget, MaskMetadata);
+		MaskInputDecl.Set<FPassExternalResourceDesc>({ TextureId });
+	}
 
-	FPassInputDeclArray Inputs = { InputDecl, MaskingTextureInput };
+	FPassInputDeclArray Inputs = { InputDecl, MaskInputDecl };
 	FMaskingPassProxy* Proxy = InFrameAllocator.Create<FMaskingPassProxy>(Inputs);
-	Proxy->MaskSourceChannel = static_cast<uint32>(MaskSource);
+	constexpr uint32 AlphaChannelIndex = static_cast<uint32>(ECompositeColorChannel::Alpha);
+	Proxy->MaskSourceChannel = (MaskingMode == ECompositeMaskingMode::Geometry) ? AlphaChannelIndex : static_cast<uint32>(MaskSource);
 	Proxy->bUnpremult = bUnpremult;
 	Proxy->bInvert = bInvert;
 
diff --git a/Engine/Plugins/Compositing/Composite/Source/Composite/Public/CompositeActor.h b/Engine/Plugins/Compositing/Composite/Source/Composite/Public/CompositeActor.h
--- a/Engine/Plugins/Compositing/Composite/Source/Composite/Public/CompositeActor.h
+++ b/Engine/Plugins/Compositing/Composite/Source/Composite/Public/CompositeActor.h
@@ -228,12 +228,12 @@ private:
 	FGameplayTag VPRole;
 
 public:
-	/** Convenience function used by layers to request a scene capture component. */
+	/** Convenience function used by passes or layers to request a scene capture component. */
 	template<class RetType>
-	UE_API RetType* FindOrCreateSceneCapture(const UCompositeLayerBase* InLayer, int32 InIndex = 0, FName InBaseName = NAME_None);
+	UE_API RetType* FindOrCreateSceneCapture(const UCompositePassBase* InPass, int32 InIndex = 0, FName InBaseName = NAME_None);
 
-	/** Convenience function used by layers to remove and scene capture components. */
-	UE_API void DestroySceneCaptures(const UCompositeLayerBase* InLayer);
+	/** Convenience function used by passes or layers to remove scene capture components. */
+	UE_API void DestroySceneCaptures(const UCompositePassBase* InPass);
 
 	/** Returns true if the composite is actively rendering. */
 	UE_API bool IsRendering() const;
@@ -245,9 +245,9 @@ private:
 	/** Propagate rendering state changes to layers and the composite core subsystem. */
 	void PropagateStateChange(ECompositeStateChangeType ChangeType);
 
-	/** Layer-managed scene capture components. */
-	UPROPERTY()
-	TMap<TWeakObjectPtr<const UCompositeLayerBase>, FSceneCaptureComponentArray> SceneCapturesPerLayer;
+	/** Pass-managed scene capture components (populated at runtime by FindOrCreateSceneCapture). */
+	UPROPERTY(Transient)
+	TMap<TWeakObjectPtr<const UCompositePassBase>, FSceneCaptureComponentArray> SceneCapturesPerPass;
 
 #if WITH_EDITORONLY_DATA
 	UE_DEPRECATED(5.8, "Camera has been deprecated, use CameraActor instead.")
diff --git a/Engine/Plugins/Compositing/Composite/Source/Composite/Public/Layers/CompositeLayerPlate.h b/Engine/Plugins/Compositing/Composite/Source/Composite/Public/Layers/CompositeLayerPlate.h
--- a/Engine/Plugins/Compositing/Composite/Source/Composite/Public/Layers/CompositeLayerPlate.h
+++ b/Engine/Plugins/Compositing/Composite/Source/Composite/Public/Layers/CompositeLayerPlate.h
@@ -172,7 +172,7 @@ public:
 	* 
 	* The texture is replaced with an internal render target.
 	*/
-	UPROPERTY(EditAnywhere, BlueprintReadWrite, Instanced, Category = "Composite", meta = (DisallowedClasses = "/Script/Composite.CompositeLayerBase, /Script/Composite.CompositePassMaterial, /Script/Composite.CompositePassDistortion, /Script/Composite.CompositePassFXAA, /Script/Composite.CompositePassSMAA"))
+	UPROPERTY(EditAnywhere, BlueprintReadWrite, Instanced, Category = "Composite", meta = (DisallowedClasses = "/Script/Composite.CompositeLayerBase, /Script/Composite.CompositePassMaterial, /Script/Composite.CompositePassDistortion"))
 	TArray<TObjectPtr<UCompositePassBase>> MediaPasses;
 
 	/**
diff --git a/Engine/Plugins/Compositing/Composite/Source/Composite/Public/Passes/CompositePassMasking.h b/Engine/Plugins/Compositing/Composite/Source/Composite/Public/Passes/CompositePassMasking.h
--- a/Engine/Plugins/Compositing/Composite/Source/Composite/Public/Passes/CompositePassMasking.h
+++ b/Engine/Plugins/Compositing/Composite/Source/Composite/Public/Passes/CompositePassMasking.h
@@ -3,6 +3,7 @@
 #pragma once
 
 #include "CompositePassBase.h"
+#include "Tickable.h"
 
 #include "CompositePassMasking.generated.h"
 
@@ -18,9 +19,21 @@ enum class ECompositeColorChannel : uint8
 	Alpha
 };
 
+/** Masking source mode. */
+UENUM(BlueprintType)
+enum class ECompositeMaskingMode : uint8
+{
+	Texture = 0  UMETA(ToolTip = "Use a texture as the mask source."),
+	Geometry = 1 UMETA(ToolTip = "Render selected meshes as a geometry mask via custom render pass."),
+};
+
+class AActor;
+class UPrimitiveComponent;
+class UCompositeSceneCapture2DComponent;
+
 /** Alpha masking pass. */
 UCLASS(MinimalAPI, BlueprintType, Blueprintable, EditInlineNew, CollapseCategories, meta=(DisplayName="Masking"))
-class UCompositePassMasking : public UCompositePassBase
+class UCompositePassMasking : public UCompositePassBase, public FTickableGameObject
 {
 	GENERATED_BODY()
 
@@ -35,15 +48,85 @@ public:
 
 	UE_API virtual FCompositeCorePassProxy* GetProxy(const UE::CompositeCore::FPassInputDecl& InputDecl, FCompositeTraversalContext& InContext, FSceneRenderingBulkObjectAllocator& InFrameAllocator) const override;
 
+	//~ Begin UObject Interface
+	UE_API virtual void BeginDestroy() override;
+
+#if WITH_EDITOR
+	UE_API virtual void PreEditChange(FProperty* PropertyAboutToChange) override;
+	UE_API virtual void PostEditChangeProperty(FPropertyChangedEvent& PropertyChangedEvent) override;
+#endif
+	//~ End UObject Interface
+
+	//~ Begin FTickableGameObject Interface
+	UE_API virtual void Tick(float DeltaTime) override;
+	UE_API virtual bool IsTickableWhenPaused() const override { return true; }
+	UE_API virtual bool IsTickableInEditor() const override { return true; }
+	UE_API virtual ETickableTickType GetTickableTickType() const override;
+	UE_API virtual bool IsTickable() const override;
+	UE_API virtual UWorld* GetTickableGameObjectWorld() const override;
+	UE_API virtual TStatId GetStatId() const override;
+	//~ End FTickableGameObject Interface
+
+	/** Get the masking mode. */
+	UFUNCTION(BlueprintGetter)
+	UE_API ECompositeMaskingMode GetMaskingMode() const { return MaskingMode; }
+
+	/** Set the masking mode. */
+	UFUNCTION(BlueprintSetter)
+	UE_API void SetMaskingMode(ECompositeMaskingMode InMaskingMode);
+
+	/** Get mask actors. */
+	UFUNCTION(BlueprintGetter)
+	UE_API const TArray<TSoftObjectPtr<AActor>>& GetMaskActors() const;
+
+	/** Set mask actors. */
+	UFUNCTION(BlueprintSetter)
+	UE_API void SetMaskActors(TArray<TSoftObjectPtr<AActor>> InMaskActors);
+
+	/** Get whether mask meshes are visible in scene capture only. */
+	UFUNCTION(BlueprintGetter)
+	UE_API bool IsVisibleInSceneCaptureOnly() const;
+
+	/** Set whether mask meshes are visible in scene capture only. */
+	UFUNCTION(BlueprintSetter)
+	UE_API void SetVisibleInSceneCaptureOnly(bool bInVisible);
+
+	/** Get the render target resolution for the geometry mask. */
+	UFUNCTION(BlueprintGetter)
+	UE_API FIntPoint GetRenderTargetResolution() const { return RenderTargetResolution; }
+
+	/** Set the render target resolution for the geometry mask. */
+	UFUNCTION(BlueprintSetter)
+	UE_API void SetRenderTargetResolution(FIntPoint InRenderTargetResolution);
+
+private:
+	/** Masking source mode. */
+	UPROPERTY(EditAnywhere, BlueprintGetter = GetMaskingMode, BlueprintSetter = SetMaskingMode, Interp, Category = "Composite", meta = (AllowPrivateAccess = true))
+	ECompositeMaskingMode MaskingMode;
+
 public:
 	/** Masking texture. */
-	UPROPERTY(EditAnywhere, BlueprintReadWrite, Category = "Composite", meta = (AllowedClasses = "/Script/Engine.Texture2D, /Script/MediaAssets.MediaTexture"))
+	UPROPERTY(EditAnywhere, BlueprintReadWrite, Category = "Composite", meta = (AllowedClasses = "/Script/Engine.Texture2D, /Script/MediaAssets.MediaTexture", EditCondition = "MaskingMode == ECompositeMaskingMode::Texture", EditConditionHides))
 	TObjectPtr<UTexture> MaskTexture;
 
 	/** Color channel used as the mask source. */
-	UPROPERTY(BlueprintReadWrite, EditAnywhere, Interp, Category = "Composite")
+	UPROPERTY(BlueprintReadWrite, EditAnywhere, Interp, Category = "Composite", meta = (EditCondition = "MaskingMode == ECompositeMaskingMode::Texture", EditConditionHides))
 	ECompositeColorChannel MaskSource;
 
+private:
+	/** Actors whose mesh primitives are rendered as a geometry mask. */
+	UPROPERTY(EditAnywhere, BlueprintGetter = GetMaskActors, BlueprintSetter = SetMaskActors, Category = "Composite", meta = (AllowPrivateAccess = true, EditCondition = "MaskingMode == ECompositeMaskingMode::Geometry", EditConditionHides))
+	TArray<TSoftObjectPtr<AActor>> MaskActors;
+
+	/** Render target resolution for the geometry mask scene capture. */
+	UPROPERTY(EditAnywhere, BlueprintGetter = GetRenderTargetResolution, BlueprintSetter = SetRenderTargetResolution, Interp, Category = "Composite", meta = (AllowPrivateAccess = true, EditCondition = "MaskingMode == ECompositeMaskingMode::Geometry", EditConditionHides))
+	FIntPoint RenderTargetResolution;
+
+	/** Visibility setting applied to mask actors. By default, mask meshes will only be visible in scene captures. */
+	UPROPERTY(EditAnywhere, BlueprintGetter = IsVisibleInSceneCaptureOnly, BlueprintSetter = SetVisibleInSceneCaptureOnly, Category = "Composite", meta = (AllowPrivateAccess = true, EditCondition = "MaskingMode == ECompositeMaskingMode::Geometry", EditConditionHides))
+	bool bVisibleInSceneCaptureOnly;
+
+public:
 	/** Color is divided by its alpha channel before masking, and multiplied again after. */
 	UPROPERTY(BlueprintReadWrite, EditAnywhere, Interp, Category = "Composite")
 	bool bUnpremult;
@@ -51,6 +134,23 @@ public:
 	/** Invert mask alpha. */
 	UPROPERTY(BlueprintReadWrite, EditAnywhere, Interp, Category = "Composite")
 	bool bInvert;
+
+private:
+	/** Find or create the scene capture component from the owning composite actor. Configures it as custom render pass on creation. */
+	UCompositeSceneCapture2DComponent* FindOrCreateSceneCaptureComponent() const;
+
+#if WITH_EDITORONLY_DATA
+	/** Cached masking mode from PreEditChange, used to detect transitions in PostEditChangeProperty. */
+	ECompositeMaskingMode PreEditMaskingMode = ECompositeMaskingMode::Texture;
+#endif
+
+	/** Update scene capture visibility on mask actor primitives. */
+	void UpdatePrimitiveVisibilityState(bool bInVisibleInSceneCaptureOnly) const;
+
+	/** Restore scene capture visibility on mask actor primitives to class defaults. */
+	void RestorePrimitiveVisibilityState() const;
+
+	friend class FCompositePassMaskingCustomization;
 };
 
 namespace UE
