Composure: Fixed missing primitive issue caused by CL 52594757.

#Synced-CL 52632071

diff --git a/Engine/Plugins/Compositing/Composite/Source/Composite/Private/Layers/CompositeLayerPlate.cpp b/Engine/Plugins/Compositing/Composite/Source/Composite/Private/Layers/CompositeLayerPlate.cpp
--- a/Engine/Plugins/Compositing/Composite/Source/Composite/Private/Layers/CompositeLayerPlate.cpp
+++ b/Engine/Plugins/Compositing/Composite/Source/Composite/Private/Layers/CompositeLayerPlate.cpp
@@ -920,8 +920,10 @@ int32 UCompositeLayerPlate::GetValidPrimitivesNum() const
 			continue;
 		}
 		
-		// If there is an entry in CompositeMeshPrimitives for the actor, use only the primitives in that entry; otherwise, use all primitives from the actor
-		if (const FCompositeMeshPrimitiveReference* PrimitiveReference = UE::Composite::Private::FindPrimitiveReferences(CompositeMeshPrimitives, Actor))
+		// If specific primitives are included for this actor, count only those; otherwise count all valid primitives.
+		// An existing entry with an empty list means "use all" (consistent with GetPrimitivesForActor).
+		const FCompositeMeshPrimitiveReference* PrimitiveReference = UE::Composite::Private::FindPrimitiveReferences(CompositeMeshPrimitives, Actor);
+		if (PrimitiveReference && !PrimitiveReference->CompositeMeshPrimitives.IsEmpty())
 		{
 			for (const TSoftObjectPtr<UPrimitiveComponent>& PrimitiveComponent : PrimitiveReference->CompositeMeshPrimitives)
 			{
