Composure: Refactored primitive collection into append-style helpers (AppendPrimitivesForActor, AppendIncludedPrimitivesForActor) to eliminate per-actor TArray intermediate allocations on the every-tick hot path.

#Synced-CL 52681084

diff --git a/Engine/Plugins/Compositing/Composite/Source/Composite/Private/Layers/CompositeLayerPlate.cpp b/Engine/Plugins/Compositing/Composite/Source/Composite/Private/Layers/CompositeLayerPlate.cpp
--- a/Engine/Plugins/Compositing/Composite/Source/Composite/Private/Layers/CompositeLayerPlate.cpp
+++ b/Engine/Plugins/Compositing/Composite/Source/Composite/Private/Layers/CompositeLayerPlate.cpp
@@ -743,45 +743,31 @@ FCompositeCorePassProxy* UCompositeLayerPlate::GetProxy(const UE::CompositeCore:
 	return InFrameAllocator.Create<FMergePassProxy>(MoveTemp(Inputs), GetMergeOperation(InContext), TEXT("Plate"), ELensDistortionHandling::Enabled, bBlendInDisplaySpace);
 }
 
-TArray<UPrimitiveComponent*> UCompositeLayerPlate::GetPrimitives() const
+void UCompositeLayerPlate::AppendIncludedPrimitivesForActor(const TSoftObjectPtr<AActor>& InActor, TArray<UPrimitiveComponent*>& OutPrimitiveComponents) const
 {
-	const UCompositeCorePluginSettings* Settings = GetDefault<UCompositeCorePluginSettings>();
-	check(Settings);
-
-	TArray<UPrimitiveComponent*> OutPrimitiveComponents;
-
-	for (const TSoftObjectPtr<AActor>& Actor : CompositeMeshes)
+	if (const FCompositeMeshPrimitiveReference* PrimitiveReference = UE::Composite::Private::FindPrimitiveReferences(CompositeMeshPrimitives, InActor))
 	{
-		if (!Actor.IsValid())
+		for (const TSoftObjectPtr<UPrimitiveComponent>& PrimitiveComponent : PrimitiveReference->CompositeMeshPrimitives)
 		{
-			continue;
+			if (PrimitiveComponent.IsValid())
+			{
+				OutPrimitiveComponents.Add(PrimitiveComponent.Get());
+			}
 		}
-
-		OutPrimitiveComponents.Append(GetPrimitivesForActor(Actor));
 	}
-
-	return OutPrimitiveComponents;
 }
 
-TArray<UPrimitiveComponent*> UCompositeLayerPlate::GetPrimitivesForActor(TSoftObjectPtr<AActor> InActor) const
+void UCompositeLayerPlate::AppendPrimitivesForActor(const TSoftObjectPtr<AActor>& InActor, TArray<UPrimitiveComponent*>& OutPrimitiveComponents) const
 {
-	if (!InActor.IsValid())
-	{
-		return TArray<UPrimitiveComponent*>();
-	}
-	
-	if (!CompositeMeshes.Contains(InActor))
-	{
-		return TArray<UPrimitiveComponent*>();
-	}
-	
-	const UCompositeCorePluginSettings* Settings = GetDefault<UCompositeCorePluginSettings>();
-	check(Settings);
-	
-	TArray<UPrimitiveComponent*> OutPrimitiveComponents = GetIncludedPrimitivesForActor(InActor);
-	if (OutPrimitiveComponents.IsEmpty())
+	const int32 StartNum = OutPrimitiveComponents.Num();
+	AppendIncludedPrimitivesForActor(InActor, OutPrimitiveComponents);
+
+	if (OutPrimitiveComponents.Num() == StartNum)
 	{
-		// If no primitives were specifically included for the actor, return all valid primitives for the actor
+		// No specific primitives included — fall back to all valid primitives for the actor
+		const UCompositeCorePluginSettings* Settings = GetDefault<UCompositeCorePluginSettings>();
+		check(Settings);
+
 		TInlineComponentArray<UPrimitiveComponent*> PrimComponents(InActor.Get());
 
 		for (UPrimitiveComponent* PrimitiveComponent : PrimComponents)
@@ -792,36 +778,46 @@ TArray<UPrimitiveComponent*> UCompositeLayerPlate::GetPrimitivesForActor(TSoftOb
 			}
 		}
 	}
+}
+
+TArray<UPrimitiveComponent*> UCompositeLayerPlate::GetPrimitives() const
+{
+	TArray<UPrimitiveComponent*> OutPrimitiveComponents;
+
+	for (const TSoftObjectPtr<AActor>& Actor : CompositeMeshes)
+	{
+		if (!Actor.IsValid())
+		{
+			continue;
+		}
+
+		AppendPrimitivesForActor(Actor, OutPrimitiveComponents);
+	}
 
 	return OutPrimitiveComponents;
 }
 
-TArray<UPrimitiveComponent*> UCompositeLayerPlate::GetIncludedPrimitivesForActor(TSoftObjectPtr<AActor> InActor) const
+TArray<UPrimitiveComponent*> UCompositeLayerPlate::GetPrimitivesForActor(TSoftObjectPtr<AActor> InActor) const
 {
-	if (!InActor.IsValid())
+	if (!InActor.IsValid() || !CompositeMeshes.Contains(InActor))
 	{
 		return TArray<UPrimitiveComponent*>();
 	}
-	
-	if (!CompositeMeshes.Contains(InActor))
+
+	TArray<UPrimitiveComponent*> OutPrimitiveComponents;
+	AppendPrimitivesForActor(InActor, OutPrimitiveComponents);
+	return OutPrimitiveComponents;
+}
+
+TArray<UPrimitiveComponent*> UCompositeLayerPlate::GetIncludedPrimitivesForActor(TSoftObjectPtr<AActor> InActor) const
+{
+	if (!InActor.IsValid() || !CompositeMeshes.Contains(InActor))
 	{
 		return TArray<UPrimitiveComponent*>();
 	}
-	
+
 	TArray<UPrimitiveComponent*> OutPrimitiveComponents;
-	if (const FCompositeMeshPrimitiveReference* PrimitiveReference = UE::Composite::Private::FindPrimitiveReferences(CompositeMeshPrimitives, InActor))
-	{
-		for (const TSoftObjectPtr<UPrimitiveComponent>& PrimitiveComponent : PrimitiveReference->CompositeMeshPrimitives)
-		{
-			if (!PrimitiveComponent.IsValid())
-			{
-				continue;
-			}
-		
-			OutPrimitiveComponents.Add(PrimitiveComponent.Get());
-		}
-	}
-	
+	AppendIncludedPrimitivesForActor(InActor, OutPrimitiveComponents);
 	return OutPrimitiveComponents;
 }
 
diff --git a/Engine/Plugins/Compositing/Composite/Source/Composite/Public/Layers/CompositeLayerPlate.h b/Engine/Plugins/Compositing/Composite/Source/Composite/Public/Layers/CompositeLayerPlate.h
--- a/Engine/Plugins/Compositing/Composite/Source/Composite/Public/Layers/CompositeLayerPlate.h
+++ b/Engine/Plugins/Compositing/Composite/Source/Composite/Public/Layers/CompositeLayerPlate.h
@@ -233,6 +233,12 @@ public:
 	bool bRemoveOverscan = true;
 
 private:
+	/** Append specifically-included primitives for InActor to OutPrimitiveComponents. Empty if none are specifically included. */
+	void AppendIncludedPrimitivesForActor(const TSoftObjectPtr<AActor>& InActor, TArray<UPrimitiveComponent*>& OutPrimitiveComponents) const;
+
+	/** Append effective primitives for InActor to OutPrimitiveComponents (included, or all valid as fallback). */
+	void AppendPrimitivesForActor(const TSoftObjectPtr<AActor>& InActor, TArray<UPrimitiveComponent*>& OutPrimitiveComponents) const;
+
 	/** Condtionally register composite mesh primitives: active holdout and add to the scene view extension custom render pass. */
 	void TryRegisterCompositePrimitives(const TArray<UPrimitiveComponent*>& InPrimitives, const UWorld* SourceWorld = nullptr) const;
 
