Composure: Allow post-process materials as media passes on plate layers.

- Remove CompositePassMaterial from MediaPasses DisallowedClasses
- Proxy creates dummy SceneTextures and minimal ViewUniformBuffer for standalone execution
- Forward OverrideOutput so the pass integrates with the preprocessing chain
- Add bIsPreprocessing to FCompositeTraversalContext to skip scene texture dependency flagging

#Synced-CL 52697243

diff --git a/Engine/Plugins/Compositing/Composite/Source/Composite/Private/Layers/CompositeLayerPlate.cpp b/Engine/Plugins/Compositing/Composite/Source/Composite/Private/Layers/CompositeLayerPlate.cpp
--- a/Engine/Plugins/Compositing/Composite/Source/Composite/Private/Layers/CompositeLayerPlate.cpp
+++ b/Engine/Plugins/Compositing/Composite/Source/Composite/Private/Layers/CompositeLayerPlate.cpp
@@ -656,6 +656,8 @@ FCompositeCorePassProxy* UCompositeLayerPlate::GetProxy(const UE::CompositeCore:
 		
 		const bool bImplicitOverscanRemoval = bRemoveOverscan && IsValid(OverscanPass) && OverscanPass->GetIsActive();
 
+		TGuardValue<bool> PreprocessingGuard(InContext.bIsPreprocessing, true);
+
 		if (bImplicitOverscanRemoval)
 		{
 			// Build combined media pass array: overscan (implicit) + user passes.
diff --git a/Engine/Plugins/Compositing/Composite/Source/Composite/Private/Passes/CompositePassMaterial.cpp b/Engine/Plugins/Compositing/Composite/Source/Composite/Private/Passes/CompositePassMaterial.cpp
--- a/Engine/Plugins/Compositing/Composite/Source/Composite/Private/Passes/CompositePassMaterial.cpp
+++ b/Engine/Plugins/Compositing/Composite/Source/Composite/Private/Passes/CompositePassMaterial.cpp
@@ -7,6 +7,7 @@
 #include "RenderGraphUtils.h"
 #include "SceneView.h"
 #include "ScreenPass.h"
+#include "SystemTextures.h"
 #include "PostProcess/PostProcessMaterialInputs.h"
 
 DECLARE_GPU_STAT_NAMED(FCompositeMaterial, TEXT("Composite.Material"));
@@ -21,6 +22,9 @@ public:
 	/* Post-process material weak pointer. */
 	TWeakObjectPtr<UMaterialInterface> Material;
 
+	/** Whether this proxy runs as a standalone pass (outside the normal post-process pipeline). */
+	bool bStandalone = false;
+
 	UE::CompositeCore::FPassTexture Add(FRDGBuilder& GraphBuilder, const FSceneView& InView, const UE::CompositeCore::FPassInputArray& Inputs, const UE::CompositeCore::FPassContext& PassContext) const override
 	{
 		RDG_EVENT_SCOPE_STAT(GraphBuilder, FCompositeMaterial, "Composite.Material");
@@ -28,9 +32,63 @@ public:
 		check(ValidateInputs(Inputs));
 
 		const UE::CompositeCore::FResourceMetadata& Metadata = Inputs[0].Metadata;
-		FPostProcessMaterialInputs PostInputs = Inputs.ToPostProcessInputs(GraphBuilder, PassContext.SceneTextures);
-		// Specify the scene color output format in case a render target needs to be created and our input is render target incompatible.
-		PostInputs.OutputFormat = FSceneTexturesConfig::Get().ColorFormat;
+
+		FSceneTextureShaderParameters SceneTextures = PassContext.SceneTextures;
+		if (bStandalone)
+		{
+			// Create dummy scene textures for preprocessing — the full scene hasn't rendered yet.
+			// We build these manually using GSystemTextures rather than calling
+			// CreateSceneTextureShaderParameters, which depends on FRDGSystemTextures that may
+			// not be initialized yet (and creating them would conflict with the renderer later).
+			FSceneTextureUniformParameters* DummyParams = GraphBuilder.AllocParameters<FSceneTextureUniformParameters>();
+			const FRDGTextureRef BlackDummy = GSystemTextures.GetBlackDummy(GraphBuilder);
+			const FRDGTextureRef DepthDummy = GSystemTextures.GetDepthDummy(GraphBuilder);
+			DummyParams->PointClampSampler = TStaticSamplerState<SF_Point>::GetRHI();
+			DummyParams->SceneColorTexture = BlackDummy;
+			DummyParams->SceneDepthTexture = DepthDummy;
+			DummyParams->ScenePartialDepthTexture = DepthDummy;
+			DummyParams->GBufferATexture = BlackDummy;
+			DummyParams->GBufferBTexture = BlackDummy;
+			DummyParams->GBufferCTexture = BlackDummy;
+			DummyParams->GBufferDTexture = BlackDummy;
+			DummyParams->GBufferETexture = BlackDummy;
+			DummyParams->GBufferFTexture = BlackDummy;
+			DummyParams->GBufferVelocityTexture = BlackDummy;
+			DummyParams->GBufferSGGXTexture = BlackDummy;
+			DummyParams->ScreenSpaceAOTexture = BlackDummy;
+			DummyParams->CustomDepthTexture = DepthDummy;
+			DummyParams->CustomStencilTexture = GraphBuilder.CreateSRV(FRDGTextureSRVDesc(GSystemTextures.GetStencilDummy(GraphBuilder)));
+			SceneTextures.SceneTextures = GraphBuilder.CreateUniformBuffer(DummyParams);
+
+			// Ensure the view has a valid uniform buffer. During early preprocessing (e.g. scene captures),
+			// the ViewUniformBuffer may not be initialized yet. Create a minimal one with essential fields
+			// so that post-process material UV computation works correctly.
+			if (!InView.ViewUniformBuffer.IsValid())
+			{
+				const FIntPoint Extent = Inputs[0].Texture.Texture ? Inputs[0].Texture.Texture->Desc.Extent : FIntPoint(1, 1);
+				const float Width = static_cast<float>(Extent.X);
+				const float Height = static_cast<float>(Extent.Y);
+
+				FViewUniformShaderParameters MinimalViewParameters;
+				MinimalViewParameters.ViewSizeAndInvSize = FVector4f(Width, Height, 1.0f / Width, 1.0f / Height);
+				MinimalViewParameters.BufferSizeAndInvSize = FVector4f(Width, Height, 1.0f / Width, 1.0f / Height);
+				MinimalViewParameters.PreExposure = 1.0f;
+				MinimalViewParameters.OneOverPreExposure = 1.0f;
+
+				const_cast<FSceneView&>(InView).ViewUniformBuffer =
+					TUniformBufferRef<FViewUniformShaderParameters>::CreateUniformBufferImmediate(MinimalViewParameters, UniformBuffer_SingleFrame);
+			}
+		}
+
+		FPostProcessMaterialInputs PostInputs = Inputs.ToPostProcessInputs(GraphBuilder, SceneTextures);
+		PostInputs.OverrideOutput = Inputs.OverrideOutput;
+		PostInputs.bStandalonePass = bStandalone;
+
+		if (!bStandalone)
+		{
+			// Specify the scene color output format in case a render target needs to be created and our input is render target incompatible.
+			PostInputs.OutputFormat = FSceneTexturesConfig::Get().ColorFormat;
+		}
 
 		constexpr bool bEvenIfPendingKill = false;
 		constexpr bool bThreadsafeTest = true;
@@ -63,10 +121,14 @@ FCompositeCorePassProxy* UCompositePassMaterial::GetProxy(const UE::CompositeCor
 {
 	FCompositePassMaterialProxy* Proxy = InFrameAllocator.Create<FCompositePassMaterialProxy>(UE::CompositeCore::FPassInputDeclArray{ InputDecl });
 	Proxy->Material = PostProcessMaterial;
+	Proxy->bStandalone = InContext.bIsPreprocessing;
 
-	// We conservatively mark any post-processing materials as potentially using scene textures.
-	// TBD if material inputs could be individually checked instead.
-	InContext.bNeedsSceneTextures = true;
+	// Conservatively mark post-processing materials as potentially using scene textures,
+	// unless we are collecting passes for preprocessing where scene textures are not available.
+	if (!InContext.bIsPreprocessing)
+	{
+		InContext.bNeedsSceneTextures = true;
+	}
 
 	return Proxy;
 }
diff --git a/Engine/Plugins/Compositing/Composite/Source/Composite/Public/Layers/CompositeLayerPlate.h b/Engine/Plugins/Compositing/Composite/Source/Composite/Public/Layers/CompositeLayerPlate.h
--- a/Engine/Plugins/Compositing/Composite/Source/Composite/Public/Layers/CompositeLayerPlate.h
+++ b/Engine/Plugins/Compositing/Composite/Source/Composite/Public/Layers/CompositeLayerPlate.h
@@ -172,7 +172,7 @@ public:
 	* 
 	* The texture is replaced with an internal render target.
 	*/
-	UPROPERTY(EditAnywhere, BlueprintReadWrite, Instanced, Category = "Composite", meta = (DisallowedClasses = "/Script/Composite.CompositeLayerBase, /Script/Composite.CompositePassMaterial, /Script/Composite.CompositePassDistortion"))
+	UPROPERTY(EditAnywhere, BlueprintReadWrite, Instanced, Category = "Composite", meta = (DisallowedClasses = "/Script/Composite.CompositeLayerBase, /Script/Composite.CompositePassDistortion"))
 	TArray<TObjectPtr<UCompositePassBase>> MediaPasses;
 
 	/**
diff --git a/Engine/Plugins/Compositing/Composite/Source/Composite/Public/Passes/CompositePassBase.h b/Engine/Plugins/Compositing/Composite/Source/Composite/Public/Passes/CompositePassBase.h
--- a/Engine/Plugins/Compositing/Composite/Source/Composite/Public/Passes/CompositePassBase.h
+++ b/Engine/Plugins/Compositing/Composite/Source/Composite/Public/Passes/CompositePassBase.h
@@ -27,6 +27,9 @@ struct FCompositeTraversalContext
 	 */
 	bool bNeedsSceneTextures = false;
 
+	/** Flag to indicate that passes are being collected for preprocessing (before the main render). */
+	bool bIsPreprocessing = false;
+
 	/** Array of passes applied at the start of rendering. */
 	TSortedMap<UE::CompositeCore::ResourceId, TArray<const FCompositeCorePassProxy*>> PreprocessingPasses;
 
diff --git a/Engine/Plugins/Compositing/Composite/Source/CompositeEditor/Private/UI/CompositePlatePassListOwner.cpp b/Engine/Plugins/Compositing/Composite/Source/CompositeEditor/Private/UI/CompositePlatePassListOwner.cpp
--- a/Engine/Plugins/Compositing/Composite/Source/CompositeEditor/Private/UI/CompositePlatePassListOwner.cpp
+++ b/Engine/Plugins/Compositing/Composite/Source/CompositeEditor/Private/UI/CompositePlatePassListOwner.cpp
@@ -5,7 +5,6 @@
 #include "CompositeEditorStyle.h"
 #include "PropertyEditorUtils.h"
 #include "Layers/CompositeLayerPlate.h"
-#include "Passes/CompositePassMaterial.h"
 #include "UObject/Package.h"
 
 #define LOCTEXT_NAMESPACE "FCompositePlatePassListOwner"
@@ -106,15 +105,7 @@ ICompositePassListOwner::FGroupFilterConfig FCompositePlatePassListOwner::GetGro
 
 int32 FCompositePlatePassListOwner::GetDefaultGroupForNewPass(const UClass* InPassClass) const
 {
-	// The default group to add the pass to is the media passes list, unless it is a PPM pass, which can only get added to the layers passes list
-	if (InPassClass->IsChildOf<UCompositePassMaterial>())
-	{
-		return (int32)EPassType::Layer;
-	}
-	else
-	{
-		return (int32)EPassType::Media;
-	}
+	return (int32)EPassType::Media;
 }
 
 bool FCompositePlatePassListOwner::CanAddPass(const UClass* InPassClass, int32 InGroupIndex, int32 InPassIndex) const
