Composure: Added ability to add custom materials to composite mesh actors in the composite mesh actor picker. Also added ability to toggle columns on and off for the picker, and option to sync selection to editor selection.

#Synced-CL 52713675

diff --git a/Engine/Plugins/Compositing/Composite/Source/CompositeEditor/Private/CompositeEditorStyle.cpp b/Engine/Plugins/Compositing/Composite/Source/CompositeEditor/Private/CompositeEditorStyle.cpp
--- a/Engine/Plugins/Compositing/Composite/Source/CompositeEditor/Private/CompositeEditorStyle.cpp
+++ b/Engine/Plugins/Compositing/Composite/Source/CompositeEditor/Private/CompositeEditorStyle.cpp
@@ -56,4 +56,9 @@ FCompositeEditorStyle::FCompositeEditorStyle()
 	Set("CompositeEditor.MediaProfile", new IMAGE_BRUSH_SVG("Editor/Icons/MediaProfile", Icon16x16));
 	Set("CompositeEditor.CaptureCleanPlate", new IMAGE_BRUSH_SVG("Editor/Icons/TextureCreate_16", Icon16x16));
 	Set("CompositeEditor.CancelCapture", new IMAGE_BRUSH_SVG("Editor/Icons/TextureCancel_16", Icon16x16));
+	
+	Set("CompositeEditor.LitMaterial", new CORE_IMAGE_BRUSH_SVG("Starship/Common/LitCube", Icon16x16));
+	Set("CompositeEditor.UnlitMaterial", new CORE_IMAGE_BRUSH_SVG("Starship/Common/UnlitCube", Icon16x16));
+	Set("CompositeEditor.CustomMaterial", new CORE_IMAGE_BRUSH_SVG("Starship/AssetIcons/Material_16", Icon16x16));
+	Set("CompositeEditor.SyncSelection", new CORE_IMAGE_BRUSH("Icons/GeneralTools/SelectAll_40x", Icon16x16));
 }
diff --git a/Engine/Plugins/Compositing/Composite/Source/CompositeEditor/Private/Customizations/CompositeLayerPlateCustomization.cpp b/Engine/Plugins/Compositing/Composite/Source/CompositeEditor/Private/Customizations/CompositeLayerPlateCustomization.cpp
--- a/Engine/Plugins/Compositing/Composite/Source/CompositeEditor/Private/Customizations/CompositeLayerPlateCustomization.cpp
+++ b/Engine/Plugins/Compositing/Composite/Source/CompositeEditor/Private/Customizations/CompositeLayerPlateCustomization.cpp
@@ -4,13 +4,16 @@
 
 #include "CompositeActor.h"
 #include "CompositeCoreSettings.h"
+#include "CompositeDepthMeshActor.h"
 #include "CompositeEditorStyle.h"
 #include "CompositeSkySphereActor.h"
 #include "CompositeMeshActor.h"
+#include "ContentBrowserModule.h"
 #include "DetailCategoryBuilder.h"
 #include "DetailLayoutBuilder.h"
 #include "DetailWidgetRow.h"
 #include "Editor.h"
+#include "IContentBrowserSingleton.h"
 #include "Engine/Texture2D.h"
 #include "IDetailGroup.h"
 #include "IPropertyUtilities.h"
@@ -26,6 +29,7 @@
 #include "EngineUtils.h"
 #include "Framework/MultiBox/MultiBoxBuilder.h"
 #include "Layers/CompositeLayerPlate.h"
+#include "Materials/MaterialInstance.h"
 #include "Styling/SlateIconFinder.h"
 #include "UI/CompositePlatePassListOwner.h"
 #include "UI/SCompositeActorPickerTable.h"
@@ -42,6 +46,7 @@ namespace CompositeLayerPlateCustomization
 	DECLARE_DELEGATE_TwoParams(FOnCustomizePropertyRow, const TSharedPtr<IPropertyHandle>&, IDetailPropertyRow&);
 	
 	static FName CompositeMeshListColumn_Components = "Components";
+	static FName CompositeMeshListColumn_Material = "Material";
 }
 
 TSharedRef<IDetailCustomization> FCompositeLayerPlateCustomization::MakeInstance()
@@ -103,7 +108,14 @@ void FCompositeLayerPlateCustomization::CustomizeDetails(IDetailLayoutBuilder& D
 
 		IDetailGroup& ActorListGroup = PlateCategory.AddGroup("CompositeMeshContent", LOCTEXT("CompositeMeshGroupName", "Composite Mesh Content"), false, true);
 		FCompositeActorPickerListRef ActorListRef(Plate, GET_MEMBER_NAME_CHECKED(UCompositeLayerPlate, CompositeMeshes), &Plate->CompositeMeshes);
-		
+		ActorListRef.OnActorsAdded.BindSP(this, &FCompositeLayerPlateCustomization::OnCompositeMeshActorsAdded);
+
+		TArray<FName> HiddenColumns;
+		if (const UCompositeEditorPanelSettings* Settings = GetDefault<UCompositeEditorPanelSettings>())
+		{
+			HiddenColumns = Settings->CompositeMeshTableHiddenColumns;
+		}
+	
 		ActorListGroup.SetToolTip(ActorListRef.GetToolTipText());
 		ActorListGroup.AddWidgetRow()
 		.WholeRowContent()
@@ -114,12 +126,34 @@ void FCompositeLayerPlateCustomization::CustomizeDetails(IDetailLayoutBuilder& D
 				constexpr bool bExcludeCompositeMeshActors = false;
 				return SCompositeActorPickerTable::MakeDefaultSceneOutlinerFilters(bExcludeCompositeMeshActors);
 			})
+			.HiddenColumnsList(HiddenColumns)
+			.OnHiddenColumnsListChanged(this, &FCompositeLayerPlateCustomization::OnHiddenColumnsListChanged)
 			.OnGenerateHeaderColumns_Lambda([](TSharedRef<SHeaderRow> InHeaderRow)
 			{
-				InHeaderRow->AddColumn(
-				SHeaderRow::Column(CompositeLayerPlateCustomization::CompositeMeshListColumn_Components)
-				.DefaultLabel(LOCTEXT("ComponentsColumnLabel", "Components"))
-				.FillWidth(0.4));
+				TArray<FName> HiddenColumns;
+				if (const UCompositeEditorPanelSettings* Settings = GetDefault<UCompositeEditorPanelSettings>())
+				{
+					HiddenColumns = Settings->CompositeMeshTableHiddenColumns;
+				}
+				
+				// .HiddenColumnsList workaround:
+				// simulate what SHeaderRow::AddColumn( const FColumn::FArguments& NewColumnArgs ) does but allowing us to modify the bIsVisible property
+				// Memory handling (delete) is done by TIndirectArray<FColumn> Columns; defined in SHeaderRow
+				{
+					SHeaderRow::FColumn* NewColumn = new SHeaderRow::FColumn(SHeaderRow::Column(CompositeLayerPlateCustomization::CompositeMeshListColumn_Material)
+						.DefaultLabel(LOCTEXT("MaterialColumnLabel", "Material"))
+						.FillWidth(0.4));
+					NewColumn->bIsVisible = !HiddenColumns.Contains(CompositeLayerPlateCustomization::CompositeMeshListColumn_Material);
+					InHeaderRow->AddColumn(*NewColumn);
+				}
+				
+				{
+					SHeaderRow::FColumn* NewColumn = new SHeaderRow::FColumn(SHeaderRow::Column(CompositeLayerPlateCustomization::CompositeMeshListColumn_Components)
+						.DefaultLabel(LOCTEXT("ComponentsColumnLabel", "Components"))
+						.FillWidth(0.4));
+					NewColumn->bIsVisible = !HiddenColumns.Contains(CompositeLayerPlateCustomization::CompositeMeshListColumn_Components);
+					InHeaderRow->AddColumn(*NewColumn);
+				}
 			})
 			.OnGenerateColumnWidget(this, &FCompositeLayerPlateCustomization::GetCompositeMeshColumnWidget)
 			.OnExtendAddMenu(this, &FCompositeLayerPlateCustomization::OnExtendCompositeMeshAddMenu)
@@ -222,8 +256,28 @@ void FCompositeLayerPlateCustomization::CustomizeTexturePropertyRow(const TShare
 	];
 }
 
+void FCompositeLayerPlateCustomization::OnHiddenColumnsListChanged()
+{
+	TSharedPtr<SCompositeActorPickerTable> PinnedCompositeMeshPicker = CompositeMeshPicker.Pin();
+	if (!PinnedCompositeMeshPicker.IsValid())
+	{
+		return;
+	}
+
+	if (UCompositeEditorPanelSettings* Settings = GetMutableDefault<UCompositeEditorPanelSettings>())
+	{
+		Settings->CompositeMeshTableHiddenColumns = PinnedCompositeMeshPicker->GetHiddenColumnsList();
+		Settings->SaveConfig();
+	}
+}
+
 TSharedRef<SWidget> FCompositeLayerPlateCustomization::GetCompositeMeshColumnWidget(const TSoftObjectPtr<AActor>& InActor, const FName& InColumnName)
 {
+	if (InColumnName == CompositeLayerPlateCustomization::CompositeMeshListColumn_Material)
+	{
+		return SNew(STextBlock).Text(this, &FCompositeLayerPlateCustomization::GetCompositeMeshMaterialLabel, InActor);	
+	}
+	
 	if (InColumnName == CompositeLayerPlateCustomization::CompositeMeshListColumn_Components)
 	{
 		return SNew(STextBlock).Text(this, &FCompositeLayerPlateCustomization::GetCompositeMeshComponentsLabel, InActor);
@@ -232,6 +286,61 @@ TSharedRef<SWidget> FCompositeLayerPlateCustomization::GetCompositeMeshColumnWid
 	return SNullWidget::NullWidget;
 }
 
+FText FCompositeLayerPlateCustomization::GetCompositeMeshMaterialLabel(TSoftObjectPtr<AActor> InActor) const
+{
+	if (!InActor.IsValid())
+	{
+		return FText::GetEmpty();
+	}
+
+	TSharedPtr<IDetailLayoutBuilder> DetailBuilder = CachedDetailBuilder.Pin();
+	if (!DetailBuilder.IsValid())
+	{
+		return FText::GetEmpty();
+	}
+
+	TArray<TWeakObjectPtr<UCompositeLayerPlate>> Plates = DetailBuilder->GetObjectsOfTypeBeingCustomized<UCompositeLayerPlate>();
+	if (Plates.Num() != 1 || !Plates[0].IsValid())
+	{
+		return FText::GetEmpty();
+	}
+
+	UCompositeLayerPlate* Plate = Plates[0].Get();
+	TArray<UPrimitiveComponent*> CompositeMeshPrimitives = Plate->GetPrimitivesForActor(InActor);
+
+	TSet<FString> MaterialNames;
+	for (const UPrimitiveComponent* Primitive : CompositeMeshPrimitives)
+	{
+		if (UMaterialInterface* Material = Primitive->GetMaterial(0))
+		{
+			MaterialNames.Add(Material->GetName());
+		}
+	}
+
+	if (MaterialNames.Num() > 1)
+	{
+		return FText::Format(LOCTEXT("MultipleMaterialsLabelFormat", "{0} Materials"), MaterialNames.Num());
+	}
+	if (MaterialNames.Num() == 1)
+	{
+		const FString MaterialName = MaterialNames.Array()[0];
+		
+		// For the default lit and unlit materials, output friendly names
+		if (MaterialName.Contains("CompositeMesh_Lit_Masked"))
+		{
+			return LOCTEXT("LitMaskedMaterialFriendlyName", "Lit Masked");
+		}
+		if (MaterialName.Contains("CompositeMesh_Unlit_AlphaComposite"))
+		{
+			return LOCTEXT("UnlitAlphaMaterialFriendlyName", "Unlit Alpha");
+		}
+		
+		return FText::FromString(MaterialName);
+	}
+
+	return LOCTEXT("NoMaterialsLabel", "No Materials");
+}
+
 FText FCompositeLayerPlateCustomization::GetCompositeMeshComponentsLabel(TSoftObjectPtr<AActor> InActor) const
 {
 	if (!InActor.IsValid())
@@ -276,6 +385,216 @@ void FCompositeLayerPlateCustomization::OnExtendCompositeMeshAddMenu(FMenuBuilde
 			FExecuteAction::CreateSP(this, &FCompositeLayerPlateCustomization::CreateCompositeMeshActor),
 			FCanExecuteAction::CreateSP(this, &FCompositeLayerPlateCustomization::CanCreateCompositeMeshActor))
 	);
+	
+	MenuBuilder.AddMenuEntry(
+		LOCTEXT("PlaceCompositeDepthMeshActorEntryLabel", "Place Composite Depth Mesh Actor"),
+		LOCTEXT("PlaceCompositeDepthMeshActorEntryToolTip", "Place a composite depth mesh actor at the appropriate position in the level"),
+		FSlateIcon(FCompositeEditorStyle::Get().GetStyleSetName(), "CompositeEditor.Composure"),
+		FUIAction(
+			FExecuteAction::CreateSP(this, &FCompositeLayerPlateCustomization::CreateCompositeDepthMeshActor),
+			FCanExecuteAction::CreateSP(this, &FCompositeLayerPlateCustomization::CanCreateCompositeMeshActor))
+	);
+	
+	MenuBuilder.EndSection();
+	MenuBuilder.BeginSection("ApplyMaterial", LOCTEXT("ApplyMaterialAddMenuSectionLabel", "Apply Material"));
+	
+	FMenuEntryParams AutoApplyMaterialEntry;
+	AutoApplyMaterialEntry.LabelOverride = LOCTEXT("AutoApplyMaterialLabel", "Automatically Apply Material");
+	AutoApplyMaterialEntry.ToolTipOverride = LOCTEXT("AutoApplyMaterialToolTip", "Automatically applies the configured material to any actors added to this list");
+	AutoApplyMaterialEntry.UserInterfaceActionType = EUserInterfaceActionType::RadioButton;
+	AutoApplyMaterialEntry.bShouldCloseWindowAfterMenuSelection = false;
+	AutoApplyMaterialEntry.DirectActions = FUIAction(
+		FExecuteAction::CreateSP(this, &FCompositeLayerPlateCustomization::SetAutoApplyMaterial, true),
+		FCanExecuteAction(),
+		FIsActionChecked::CreateSP(this, &FCompositeLayerPlateCustomization::IsAutoApplyMaterialChecked, true));
+	MenuBuilder.AddMenuEntry(AutoApplyMaterialEntry);
+	
+	FMenuEntryParams NoModifyMaterialEntry;
+	NoModifyMaterialEntry.LabelOverride = LOCTEXT("NoModifyMaterialLabel", "Do Not Modify Actor's Material");
+	NoModifyMaterialEntry.ToolTipOverride = LOCTEXT("NoModifyMaterialToolTip", "The actor's materials are not modified when the actor is added to this list");
+	NoModifyMaterialEntry.UserInterfaceActionType = EUserInterfaceActionType::RadioButton;
+	NoModifyMaterialEntry.bShouldCloseWindowAfterMenuSelection = false;
+	NoModifyMaterialEntry.DirectActions = FUIAction(
+		FExecuteAction::CreateSP(this, &FCompositeLayerPlateCustomization::SetAutoApplyMaterial, false),
+		FCanExecuteAction(),
+		FIsActionChecked::CreateSP(this, &FCompositeLayerPlateCustomization::IsAutoApplyMaterialChecked, false));
+	MenuBuilder.AddMenuEntry(NoModifyMaterialEntry);
+	
+	MenuBuilder.AddSeparator();
+	
+	MenuBuilder.AddSubMenu(
+		TAttribute<FText>::CreateSP(this, &FCompositeLayerPlateCustomization::GetMaterialSubMenuLabel),
+		FText::GetEmpty(),
+		FNewMenuDelegate::CreateSP(this, &FCompositeLayerPlateCustomization::CreateMaterialSubMenu),
+		false,
+		FSlateIcon(),
+		false
+	);
+}
+
+void FCompositeLayerPlateCustomization::SetAutoApplyMaterial(bool bInAutoApplyMaterial)
+{
+	if (UCompositeEditorPanelSettings* Settings = GetMutableDefault<UCompositeEditorPanelSettings>())
+	{
+		Settings->bAutoApplyCompositeMeshMaterial = bInAutoApplyMaterial;
+		Settings->SaveConfig();
+	}
+}
+
+bool FCompositeLayerPlateCustomization::IsAutoApplyMaterialChecked(bool bInAutoApplyMaterial) const
+{
+	if (const UCompositeEditorPanelSettings* Settings = GetDefault<UCompositeEditorPanelSettings>())
+	{
+		return Settings->bAutoApplyCompositeMeshMaterial == bInAutoApplyMaterial;
+	}
+	
+	return false;
+}
+
+FText FCompositeLayerPlateCustomization::GetMaterialSubMenuLabel() const
+{
+	FText MaterialName = FText::GetEmpty();
+	
+	if (const UCompositeEditorPanelSettings* Settings = GetDefault<UCompositeEditorPanelSettings>())
+	{
+		switch (Settings->CompositeMeshMaterialType)
+		{
+		case ECompositeMeshAppliedMaterialType::LitMaskedMaterial:
+			MaterialName = LOCTEXT("LitMaskedMaterialName", "Lit Masked");
+			break;
+		
+		case ECompositeMeshAppliedMaterialType::UnlitAlphaMaterial:
+			MaterialName = LOCTEXT("UnlitAlphaMaterialName", "Unlit Alpha");
+			break;
+		
+		case ECompositeMeshAppliedMaterialType::CustomMaterial:
+			if (CustomMaterialAsset.IsValid())
+			{
+				MaterialName = FText::FromName(CustomMaterialAsset.AssetName);
+			}
+			break;
+		}
+	}
+	
+	return FText::Format(LOCTEXT("AutoApplyMaterialSubMenuLabelFormat", "Material ({0})"), MaterialName);
+}
+
+void FCompositeLayerPlateCustomization::CreateMaterialSubMenu(FMenuBuilder& InMenuBuilder)
+{
+	InMenuBuilder.BeginSection("Material", LOCTEXT("MaterialSectionLabel", "Material"));
+	{
+		InMenuBuilder.AddMenuEntry(
+			LOCTEXT("ApplyLitMaskedMaterial", "Apply Lit Masked Material"),
+			LOCTEXT("ApplyLitMaskedMaterialToolTip", "Apply the plugin default lit masked material to selected actors: best for catching shadows and reflections."),
+			FSlateIcon(FCompositeEditorStyle::Get().GetStyleSetName(), "CompositeEditor.LitMaterial"),
+			FUIAction(
+				FExecuteAction::CreateSP(this, &FCompositeLayerPlateCustomization::SetAutoAppliedMaterialType, ECompositeMeshAppliedMaterialType::LitMaskedMaterial),
+				FCanExecuteAction(),
+				FIsActionChecked::CreateSP(this, &FCompositeLayerPlateCustomization::IsAutoAppliedMaterialTypeChecked, ECompositeMeshAppliedMaterialType::LitMaskedMaterial)
+			),
+			NAME_None,
+			EUserInterfaceActionType::RadioButton
+		);
+
+		InMenuBuilder.AddMenuEntry(
+			LOCTEXT("ApplyUnlitAlphaMaterial", "Apply Unlit Alpha Material"),
+			LOCTEXT("ApplyUnlitAlphaMaterialToolTip", "Apply the plugin default unlit alpha material to selected actors: best for keying media."),
+			FSlateIcon(FCompositeEditorStyle::Get().GetStyleSetName(), "CompositeEditor.UnlitMaterial"),
+			FUIAction(
+				FExecuteAction::CreateSP(this, &FCompositeLayerPlateCustomization::SetAutoAppliedMaterialType, ECompositeMeshAppliedMaterialType::UnlitAlphaMaterial),
+				FCanExecuteAction(),
+				FIsActionChecked::CreateSP(this, &FCompositeLayerPlateCustomization::IsAutoAppliedMaterialTypeChecked, ECompositeMeshAppliedMaterialType::UnlitAlphaMaterial)
+			),
+			NAME_None,
+			EUserInterfaceActionType::RadioButton
+		);
+		
+		InMenuBuilder.AddMenuEntry(
+			LOCTEXT("ApplyCustomMaterial", "Apply Custom Material"),
+			LOCTEXT("ApplyCustomMaterialToolTip", "Apply a selected custom material"),
+			FSlateIcon(FCompositeEditorStyle::Get().GetStyleSetName(), "CompositeEditor.CustomMaterial"),
+			FUIAction(
+				FExecuteAction::CreateSP(this, &FCompositeLayerPlateCustomization::SetAutoAppliedMaterialType, ECompositeMeshAppliedMaterialType::CustomMaterial),
+				FCanExecuteAction::CreateSP(this, &FCompositeLayerPlateCustomization::CanApplyCustomMaterial),
+				FIsActionChecked::CreateSP(this, &FCompositeLayerPlateCustomization::IsAutoAppliedMaterialTypeChecked, ECompositeMeshAppliedMaterialType::CustomMaterial)
+			),
+			NAME_None,
+			EUserInterfaceActionType::RadioButton
+		);
+		
+		InMenuBuilder.AddSeparator();
+		InMenuBuilder.AddMenuEntry(
+			TAttribute<FText>::CreateSP(this, &FCompositeLayerPlateCustomization::GetCustomMaterialEntryLabel),
+			FText::GetEmpty(),
+			FSlateIcon(FCompositeEditorStyle::Get().GetStyleSetName(), "CompositeEditor.CustomMaterial"),
+			FUIAction(FExecuteAction(), FCanExecuteAction::CreateLambda([] { return false; })));
+	}
+	InMenuBuilder.EndSection();
+
+	InMenuBuilder.BeginSection("MaterialBrowser", LOCTEXT("MaterialBrowserSectionLabel", "Browse"));
+	{
+		FAssetPickerConfig AssetPickerConfig;
+		AssetPickerConfig.Filter.ClassPaths.Add(UMaterial::StaticClass()->GetClassPathName());
+		AssetPickerConfig.Filter.ClassPaths.Add(UMaterialInstance::StaticClass()->GetClassPathName());
+		AssetPickerConfig.Filter.bRecursiveClasses = true;
+		AssetPickerConfig.bAllowNullSelection = false;
+		AssetPickerConfig.OnAssetSelected = FOnAssetSelected::CreateSP(this, &FCompositeLayerPlateCustomization::OnCustomMaterialSelected);
+		AssetPickerConfig.InitialAssetViewType = EAssetViewType::List;
+		AssetPickerConfig.InitialAssetSelection = CustomMaterialAsset;
+
+		FContentBrowserModule& ContentBrowserModule = FModuleManager::Get().LoadModuleChecked<FContentBrowserModule>(TEXT("ContentBrowser"));
+
+		InMenuBuilder.AddWidget(
+			SNew(SBox)
+			.MaxDesiredHeight(600.0f)
+			.WidthOverride(256.0f)
+			[
+				ContentBrowserModule.Get().CreateAssetPicker(AssetPickerConfig)
+			], 
+			FText(),
+			true
+		);
+	}
+	InMenuBuilder.EndSection();
+}
+
+void FCompositeLayerPlateCustomization::SetAutoAppliedMaterialType(ECompositeMeshAppliedMaterialType InMaterialType)
+{
+	if (UCompositeEditorPanelSettings* Settings = GetMutableDefault<UCompositeEditorPanelSettings>())
+	{
+		Settings->CompositeMeshMaterialType = InMaterialType;
+		Settings->SaveConfig();
+	}
+}
+
+bool FCompositeLayerPlateCustomization::IsAutoAppliedMaterialTypeChecked(ECompositeMeshAppliedMaterialType InMaterialType) const
+{
+	if (const UCompositeEditorPanelSettings* Settings = GetDefault<UCompositeEditorPanelSettings>())
+	{
+		return Settings->CompositeMeshMaterialType == InMaterialType;
+	}
+	
+	return false;
+}
+
+FText FCompositeLayerPlateCustomization::GetCustomMaterialEntryLabel() const
+{
+	if (CustomMaterialAsset.IsValid() && CustomMaterialAsset.IsInstanceOf(UMaterialInterface::StaticClass()))
+	{
+		return FText::FromName(CustomMaterialAsset.AssetName);
+	}
+	
+	return LOCTEXT("SelectCustomMaterial", "Select a Custom Material");
+}
+
+bool FCompositeLayerPlateCustomization::CanApplyCustomMaterial() const
+{
+	return CustomMaterialAsset.IsValid() && CustomMaterialAsset.IsInstanceOf(UMaterialInterface::StaticClass());
+}
+
+void FCompositeLayerPlateCustomization::OnCustomMaterialSelected(const FAssetData& AssetData)
+{
+	CustomMaterialAsset = AssetData;
 }
 
 void FCompositeLayerPlateCustomization::OnExtendCompositeMeshContextMenu(FMenuBuilder& MenuBuilder, TArray<TSoftObjectPtr<AActor>>& SelectedActors)
@@ -468,7 +787,51 @@ bool FCompositeLayerPlateCustomization::IsCompositeMeshPrimitiveChecked(TSoftObj
 	return CompositeMeshPrimitives.Contains(InPrimitive);
 }
 
-void FCompositeLayerPlateCustomization::ApplyMaterialToActor(TSoftObjectPtr<AActor>& InActor, UMaterial* InMaterial)
+void FCompositeLayerPlateCustomization::OnCompositeMeshActorsAdded(TArray<TSoftObjectPtr<AActor>>& InActors)
+{
+	if (const UCompositeEditorPanelSettings* Settings = GetDefault<UCompositeEditorPanelSettings>())
+	{
+		if (Settings->bAutoApplyCompositeMeshMaterial)
+		{
+			UMaterialInterface* Material = nullptr;
+			switch (Settings->CompositeMeshMaterialType)
+			{
+			case ECompositeMeshAppliedMaterialType::LitMaskedMaterial:
+				Material = LoadObject<UMaterialInterface>(nullptr, TEXT("/Composite/Materials/M_CompositeMesh_Lit_Masked.M_CompositeMesh_Lit_Masked"));
+				break;
+			
+			case ECompositeMeshAppliedMaterialType::UnlitAlphaMaterial:
+				Material = LoadObject<UMaterialInterface>(nullptr, TEXT("/Composite/Materials/M_CompositeMesh_Unlit_AlphaComposite.M_CompositeMesh_Unlit_AlphaComposite"));
+				break;
+			
+			case ECompositeMeshAppliedMaterialType::CustomMaterial:
+				if (CustomMaterialAsset.IsValid() && CustomMaterialAsset.IsInstanceOf(UMaterialInterface::StaticClass()))
+				{
+					Material = Cast<UMaterialInterface>(CustomMaterialAsset.GetAsset());
+				}
+				break;
+			}
+
+			if (Material == nullptr)
+			{
+				return;
+			}
+
+			for (TSoftObjectPtr<AActor>& Actor : InActors)
+			{
+				// Skip depth mesh actors, as they have a special material
+				if (Actor->IsA<ACompositeDepthMeshActor>())
+				{
+					continue;
+				}
+				
+				ApplyMaterialToActor(Actor, Material);
+			}
+		}
+	}
+}
+
+void FCompositeLayerPlateCustomization::ApplyMaterialToActor(TSoftObjectPtr<AActor>& InActor, UMaterialInterface* InMaterial)
 {
 	if (!InActor.IsValid())
 	{
@@ -570,6 +933,11 @@ void FCompositeLayerPlateCustomization::CreateCompositeMeshActor()
 		Plate->PreEditChange(CompositeMeshesProperty);
 
 		Plate->CompositeMeshes.Add(AddedMeshActor);
+		
+		TArray<TSoftObjectPtr<AActor>> AddedActors;
+		AddedActors.Add(AddedMeshActor);
+
+		OnCompositeMeshActorsAdded(AddedActors);
 
 		FPropertyChangedEvent ChangedEvent(CompositeMeshesProperty, EPropertyChangeType::ArrayAdd);
 		Plate->PostEditChangeProperty(ChangedEvent);
@@ -595,6 +963,50 @@ void FCompositeLayerPlateCustomization::CreateCompositeMeshActor()
 	}
 }
 
+void FCompositeLayerPlateCustomization::CreateCompositeDepthMeshActor()
+{
+	TSharedPtr<IDetailLayoutBuilder> PinnedDetailBuilder = CachedDetailBuilder.Pin();
+	if (!PinnedDetailBuilder.IsValid())
+	{
+		return;
+	}
+	
+	TArray<TWeakObjectPtr<UCompositeLayerPlate>> Objects = PinnedDetailBuilder->GetObjectsOfTypeBeingCustomized<UCompositeLayerPlate>();
+	if (Objects.Num() != 1)
+	{
+		return;
+	}
+
+	TStrongObjectPtr<UCompositeLayerPlate> Plate = Objects[0].Pin();
+	if (!Plate.IsValid())
+	{
+		return;
+	}
+
+	ACompositeActor* CompositeActor = Plate->GetTypedOuter<ACompositeActor>();
+	if (!CompositeActor)
+	{
+		return;
+	}
+
+	FTransform SpawnTransform = FTransform::Identity;
+
+	FScopedTransaction AddCompositeDepthMeshTransaction(LOCTEXT("AddCompositeDepthMeshTransaction", "Add Composite Depth Mesh"));
+	ACompositeDepthMeshActor* AddedDepthMeshActor = Cast<ACompositeDepthMeshActor>(CompositeActor->GetWorld()->SpawnActor(ACompositeDepthMeshActor::StaticClass(), &SpawnTransform));
+	if (AddedDepthMeshActor)
+	{
+		Plate->Modify();
+
+		FProperty* CompositeMeshesProperty = FindFProperty<FProperty>(UCompositeLayerPlate::StaticClass(), GET_MEMBER_NAME_CHECKED(UCompositeLayerPlate, CompositeMeshes));
+		Plate->PreEditChange(CompositeMeshesProperty);
+
+		Plate->CompositeMeshes.Add(AddedDepthMeshActor);
+
+		FPropertyChangedEvent ChangedEvent(CompositeMeshesProperty, EPropertyChangeType::ArrayAdd);
+		Plate->PostEditChangeProperty(ChangedEvent);
+	}
+}
+
 bool FCompositeLayerPlateCustomization::CanCreateCompositeMeshActor() const
 {
 	TSharedPtr<IDetailLayoutBuilder> PinnedDetailBuilder = CachedDetailBuilder.Pin();
diff --git a/Engine/Plugins/Compositing/Composite/Source/CompositeEditor/Private/Customizations/CompositeLayerPlateCustomization.h b/Engine/Plugins/Compositing/Composite/Source/CompositeEditor/Private/Customizations/CompositeLayerPlateCustomization.h
--- a/Engine/Plugins/Compositing/Composite/Source/CompositeEditor/Private/Customizations/CompositeLayerPlateCustomization.h
+++ b/Engine/Plugins/Compositing/Composite/Source/CompositeEditor/Private/Customizations/CompositeLayerPlateCustomization.h
@@ -3,9 +3,11 @@
 #pragma once
 
 #include "IDetailCustomization.h"
+#include "AssetRegistry/AssetData.h"
 #include "Materials/Material.h"
 #include "UObject/SoftObjectPtr.h"
 #include "Styling/SlateTypes.h"
+#include "UI/CompositeEditorPanelSettings.h"
 
 class SWidget;
 class IDetailPropertyRow;
@@ -36,15 +38,48 @@ private:
 	/** Customizes the Texture property in UCompositeLayerPlate to add in the media profile source selector */
 	void CustomizeTexturePropertyRow(const TSharedPtr<IPropertyHandle>& InPropertyHandle, IDetailPropertyRow& InPropertyRow);
 
+	/** Raised when the user has hidden columns in the composite mesh actor picker table */
+	void OnHiddenColumnsListChanged();
+	
 	/** Creates the widget to display for each actor in the Components column of the actor picker table */
 	TSharedRef<SWidget> GetCompositeMeshColumnWidget(const TSoftObjectPtr<AActor>& InActor, const FName& InColumnName);
 
+	/** Gets the label to display for the Material column for each actor in the actor picker table */
+	FText GetCompositeMeshMaterialLabel(TSoftObjectPtr<AActor> Actor) const;
+	
 	/** Gets the label to display for each actor in the Components column of the actor picker table */
 	FText GetCompositeMeshComponentsLabel(TSoftObjectPtr<AActor> InActor) const;
 
 	/** Raised by the composite mesh actor list when building the Add Actor menu */
 	void OnExtendCompositeMeshAddMenu(FMenuBuilder& MenuBuilder);
 
+	/** Sets the AutoApplyMaterial flag to the specified value */
+	void SetAutoApplyMaterial(bool bInAutoApplyMaterial);
+	
+	/** Checks if the AutoApplyMaterial flag is set to the specified value */
+	bool IsAutoApplyMaterialChecked(bool bInAutoApplyMaterial) const;
+	
+	/** Gets the label for the Material selector sub menu in the Add menu */
+	FText GetMaterialSubMenuLabel() const;
+	
+	/** Creates the Material selector submenu in the Add menu */
+	void CreateMaterialSubMenu(FMenuBuilder& InMenuBuilder);
+	
+	/** Sets the material type of the auto-applied material */
+	void SetAutoAppliedMaterialType(ECompositeMeshAppliedMaterialType InMaterialType);
+	
+	/** Gets whether the specified material type is the currently selected type */
+	bool IsAutoAppliedMaterialTypeChecked(ECompositeMeshAppliedMaterialType InMaterialType) const;
+	
+	/** Gets the label for the custom material menu entry in the Add dropdown menu */
+	FText GetCustomMaterialEntryLabel() const;
+	
+	/** Gets whether a custom material can currently be auto-applied */
+	bool CanApplyCustomMaterial() const;
+	
+	/** Raised when a new custom material is selected from the material browser in the Material submenu */
+	void OnCustomMaterialSelected(const FAssetData& AssetData);
+	
 	/** Raised by the composite mesh actor list when building the context menu */
 	void OnExtendCompositeMeshContextMenu(FMenuBuilder& MenuBuilder, TArray<TSoftObjectPtr<AActor>>& SelectedActors);
 
@@ -62,13 +97,19 @@ private:
 	
 	/** Gets whether a composite mesh primitive entry is checked in the Component section of the context menu */
 	bool IsCompositeMeshPrimitiveChecked(TSoftObjectPtr<AActor> InActor, UPrimitiveComponent* InPrimitive) const;
-		
+
+	/** Raised when actors are added to the Plate's CompositeMeshes list */
+	void OnCompositeMeshActorsAdded(TArray<TSoftObjectPtr<AActor>>& InActors);
+
 	/** Raised when a material should be appplied to an actor in the actor picker list */
-	void ApplyMaterialToActor(TSoftObjectPtr<AActor>& InActor, UMaterial* InMaterial);
+	void ApplyMaterialToActor(TSoftObjectPtr<AActor>& InActor, UMaterialInterface* InMaterial);
 
 	/** Creates a composite mesh actor at the appropriate position in the level and adds it to the plate's composite meshes list */
 	void CreateCompositeMeshActor();
 
+	/** Creates a composite depth mesh actor at the appropriate position in the level and adds it to the plate's composite meshes list */
+	void CreateCompositeDepthMeshActor();
+
 	/** Gets whether a composite mesh actor can be created for this plate layer */
 	bool CanCreateCompositeMeshActor() const;
 
@@ -89,4 +130,7 @@ private:
 	TWeakPtr<SCompositePassTreePanel> PassPanel;
 	TWeakPtr<SCompositeActorPickerTable> CompositeMeshPicker;
 	TSharedPtr<IPropertyHandle> TexturePropertyHandle;
+	
+	/** The custom material the user has selected to auto-apply to actors added to the composite meshes list */
+	FAssetData CustomMaterialAsset = FAssetData();
 };
diff --git a/Engine/Plugins/Compositing/Composite/Source/CompositeEditor/Private/UI/CompositeEditorPanelSettings.h b/Engine/Plugins/Compositing/Composite/Source/CompositeEditor/Private/UI/CompositeEditorPanelSettings.h
new file mode 100644
--- /dev/null
+++ b/Engine/Plugins/Compositing/Composite/Source/CompositeEditor/Private/UI/CompositeEditorPanelSettings.h
@@ -0,0 +1,39 @@
+﻿// Copyright Epic Games, Inc. All Rights Reserved.
+
+#pragma once
+
+#include "CoreMinimal.h"
+#include "UObject/Object.h"
+#include "CompositeEditorPanelSettings.generated.h"
+
+UENUM()
+enum class ECompositeMeshAppliedMaterialType : uint8
+{
+	LitMaskedMaterial,
+	UnlitAlphaMaterial,
+	CustomMaterial
+};
+
+/** Settings for the Composure panel */
+UCLASS(MinimalAPI, Config=EditorPerProjectUserSettings)
+class UCompositeEditorPanelSettings : public UObject
+{
+	GENERATED_BODY()
+	
+public:
+	/** Whether the selected material is applied to actors that are added to the composite meshes list in the Plate layer */
+	UPROPERTY(Config)
+	bool bAutoApplyCompositeMeshMaterial = true;
+	
+	/** The material type to apply to actors that are added to the composite meshes list in the Plate layer */
+	UPROPERTY(Config)
+	ECompositeMeshAppliedMaterialType CompositeMeshMaterialType = ECompositeMeshAppliedMaterialType::LitMaskedMaterial;
+	
+	/** List of columns that have been hidden by the user in the composite meshes list in the Plate layer */
+	UPROPERTY(Config)
+	TArray<FName> CompositeMeshTableHiddenColumns = TArray<FName>();
+	
+	/** Indicates that actor pickers in the Composure panel will sync their selection with the editor selection */
+	UPROPERTY(Config)
+	bool bSyncActorSelection = true;
+};
diff --git a/Engine/Plugins/Compositing/Composite/Source/CompositeEditor/Private/UI/SCompositeActorPickerSceneOutliner.cpp b/Engine/Plugins/Compositing/Composite/Source/CompositeEditor/Private/UI/SCompositeActorPickerSceneOutliner.cpp
--- a/Engine/Plugins/Compositing/Composite/Source/CompositeEditor/Private/UI/SCompositeActorPickerSceneOutliner.cpp
+++ b/Engine/Plugins/Compositing/Composite/Source/CompositeEditor/Private/UI/SCompositeActorPickerSceneOutliner.cpp
@@ -303,6 +303,14 @@ void SCompositeActorPickerSceneOutliner::OnActorSelected(AActor* InActor, bool b
 		ChangeType = EPropertyChangeType::ArrayRemove;
 	}
 
+	if (bIsAdding)
+	{
+		TArray<TSoftObjectPtr<AActor>> AddedActors;
+		AddedActors.Add(InActor);
+
+		ActorListRef.OnActorsAdded.ExecuteIfBound(AddedActors);
+	}
+
 	ActorListRef.NotifyPostEditChangeList(ChangeType, LastModifiedIndex);
 	OnActorListChanged.ExecuteIfBound();
 }
@@ -325,6 +333,7 @@ void SCompositeActorPickerSceneOutliner::OnAllActorsSelected(const TArray<AActor
 
 	bool bModified = false;
 
+	TArray<TSoftObjectPtr<AActor>> AddedActors;
 	for (AActor* Actor : InActors)
 	{
 		if (!Actor)
@@ -337,6 +346,7 @@ void SCompositeActorPickerSceneOutliner::OnAllActorsSelected(const TArray<AActor
 		if (bIsSelected && !bIsInList)
 		{
 			ActorListRef.ActorList->Add(Actor);
+			AddedActors.Add(Actor);
 			bModified = true;
 		}
 		else if (!bIsSelected && bIsInList)
@@ -348,6 +358,11 @@ void SCompositeActorPickerSceneOutliner::OnAllActorsSelected(const TArray<AActor
 
 	if (bModified)
 	{
+		if (!AddedActors.IsEmpty())
+		{
+			ActorListRef.OnActorsAdded.ExecuteIfBound(AddedActors);
+		}
+
 		ActorListRef.NotifyPostEditChangeList(bIsSelected ? EPropertyChangeType::ArrayAdd : EPropertyChangeType::ArrayRemove);
 		OnActorListChanged.ExecuteIfBound();
 	}
diff --git a/Engine/Plugins/Compositing/Composite/Source/CompositeEditor/Private/UI/SCompositeActorPickerTable.cpp b/Engine/Plugins/Compositing/Composite/Source/CompositeEditor/Private/UI/SCompositeActorPickerTable.cpp
--- a/Engine/Plugins/Compositing/Composite/Source/CompositeEditor/Private/UI/SCompositeActorPickerTable.cpp
+++ b/Engine/Plugins/Compositing/Composite/Source/CompositeEditor/Private/UI/SCompositeActorPickerTable.cpp
@@ -8,8 +8,12 @@
 #include "Components/PrimitiveComponent.h"
 #include "SceneOutlinerFilters.h"
 #include "CompositeEditorCommands.h"
+#include "CompositeEditorPanelSettings.h"
+#include "CompositeEditorStyle.h"
+#include "ContentBrowserModule.h"
 #include "Editor.h"
 #include "EditorActorFolders.h"
+#include "IContentBrowserSingleton.h"
 #include "SceneOutlinerDragDrop.h"
 #include "SCompositeActorPickerSceneOutliner.h"
 #include "ScopedTransaction.h"
@@ -21,8 +25,8 @@
 #include "Filters/SFilterSearchBox.h"
 #include "Framework/MultiBox/MultiBoxBuilder.h"
 #include "Materials/Material.h"
+#include "Materials/MaterialInstance.h"
 #include "Misc/TextFilter.h"
-#include "Modules/ModuleManager.h"
 #include "Styling/SlateIconFinder.h"
 #include "Widgets/Images/SLayeredImage.h"
 #include "Widgets/Input/SCheckBox.h"
@@ -94,6 +98,39 @@ public:
 						AddIcon
 					]
 				]
+
+				+SHorizontalBox::Slot()
+				.AutoWidth()
+				.Padding(0.0f, 0.0f, 4.0f, 0.0f)
+				[
+					SNew(SCheckBox)
+					.Padding(4.0)
+					.Style(&FAppStyle::Get().GetWidgetStyle<FCheckBoxStyle>("ToggleButtonCheckbox"))
+					.Type(ESlateCheckBoxType::ToggleButton)
+					.ToolTipText(LOCTEXT("SyncSelectionToolTip", "Sync the world outliner selection with this table's selection"))
+					.IsChecked_Lambda([]
+					{
+						if (const UCompositeEditorPanelSettings* Settings = GetDefault<UCompositeEditorPanelSettings>())
+						{
+							return Settings->bSyncActorSelection ? ECheckBoxState::Checked : ECheckBoxState::Unchecked;
+						}
+
+						return ECheckBoxState::Unchecked;
+					})
+					.OnCheckStateChanged_Lambda([this](ECheckBoxState CheckState)
+					{
+						if (UCompositeEditorPanelSettings* Settings = GetMutableDefault<UCompositeEditorPanelSettings>())
+						{
+							Settings->bSyncActorSelection = CheckState == ECheckBoxState::Checked;
+							Settings->SaveConfig();
+						}
+					})
+					[
+						SNew(SImage)
+						.Image(FCompositeEditorStyle::Get().GetBrush("CompositeEditor.SyncSelection"))
+						.ColorAndOpacity(FSlateColor::UseForeground())
+					]
+				]
 				
 				+SHorizontalBox::Slot()
 				.FillWidth(1.0f)
@@ -196,6 +233,7 @@ private:
 		TSharedPtr<FScopedTransaction> AddActorsTransaction = nullptr;
 		int32 LastModifiedIndex = INDEX_NONE;
 
+		TArray<TSoftObjectPtr<AActor>> AddedActors;
 		for (AActor* SelectedActor : SelectedActors)
 		{
 			if (ActorListRef.ActorList->Contains(SelectedActor))
@@ -212,12 +250,19 @@ private:
 			}
 
 			ActorListRef.ActorList->Add(SelectedActor);
+			AddedActors.Add(SelectedActor);
 			LastModifiedIndex = ActorListRef.ActorList->Num() - 1;
 		}
 
 		if (AddActorsTransaction.IsValid())
 		{
+			if (!AddedActors.IsEmpty())
+			{
+				ActorListRef.OnActorsAdded.ExecuteIfBound(AddedActors);
+			}
+
 			ActorListRef.NotifyPostEditChangeList(EPropertyChangeType::ArrayAdd, LastModifiedIndex);
+
 			OnActorListChanged.ExecuteIfBound();
 		}
 	}
@@ -449,6 +494,8 @@ void SCompositeActorPickerTable::Construct(const FArguments& InArgs, const FComp
 	OnApplyMaterialToActor = InArgs._OnApplyMaterialToActor;
 	
 	FCoreUObjectDelegates::OnObjectPropertyChanged.AddSP(this, &SCompositeActorPickerTable::OnObjectPropertyChanged);
+	USelection::SelectionChangedEvent.AddSP(this, &SCompositeActorPickerTable::OnLevelSelectionChanged);
+	USelection::SelectObjectEvent.AddSP(this, &SCompositeActorPickerTable::OnLevelSelectionChanged);
 	
 	if (GEditor)
 	{
@@ -458,8 +505,11 @@ void SCompositeActorPickerTable::Construct(const FArguments& InArgs, const FComp
 	CommandList = MakeShared<FUICommandList>();
 	BindCommands();
 	
-	TSharedRef<SHeaderRow> HeaderRow = 
+	HeaderRow = 
 		SNew(SHeaderRow)
+		.CanSelectGeneratedColumn(true)
+		.HiddenColumnsList(InArgs._HiddenColumnsList)
+		.OnHiddenColumnsListChanged(InArgs._OnHiddenColumnsListChanged)
 
 		+SHeaderRow::Column(CompositeActorPickerTable::ActorListColumn_ActorName)
 		.FillWidth(0.6)
@@ -469,7 +519,7 @@ void SCompositeActorPickerTable::Construct(const FArguments& InArgs, const FComp
 		.FillWidth(0.4)
 		.DefaultLabel(LOCTEXT("ActorTypeColumnLabel", "Type"));
 	
-	InArgs._OnGenerateHeaderColumns.ExecuteIfBound(HeaderRow);
+	InArgs._OnGenerateHeaderColumns.ExecuteIfBound(HeaderRow.ToSharedRef());
 	
 	ChildSlot
 	[
@@ -505,6 +555,7 @@ void SCompositeActorPickerTable::Construct(const FArguments& InArgs, const FComp
 						.OnGenerateColumnWidget(OnGenerateColumnWidget);
 				})
 				.OnContextMenuOpening(this, &SCompositeActorPickerTable::CreateListContextMenu)
+				.OnSelectionChanged(this, &SCompositeActorPickerTable::OnActorSelectionChanged)
 			]
 		]
 
@@ -530,6 +581,8 @@ SCompositeActorPickerTable::~SCompositeActorPickerTable()
 	}
 	
 	FCoreUObjectDelegates::OnObjectPropertyChanged.RemoveAll(this);
+	USelection::SelectionChangedEvent.RemoveAll(this);
+	USelection::SelectObjectEvent.RemoveAll(this);
 }
 
 void SCompositeActorPickerTable::Tick(const FGeometry& AllottedGeometry, const double InCurrentTime, const float InDeltaTime)
@@ -567,6 +620,16 @@ void SCompositeActorPickerTable::PostRedo(bool bSuccess)
 	PostUndo(bSuccess);
 }
 
+TArray<FName> SCompositeActorPickerTable::GetHiddenColumnsList() const
+{
+	if (HeaderRow.IsValid())
+	{
+		return HeaderRow->GetHiddenColumnIds();
+	}
+	
+	return TArray<FName>();
+}
+
 void SCompositeActorPickerTable::BindCommands()
 {
 	CommandList->MapAction(FCompositeEditorCommands::Get().RemoveActor,
@@ -623,6 +686,34 @@ void SCompositeActorPickerTable::FillActorList()
 	if (ListView.IsValid())
 	{
 		ListView->RequestListRefresh();
+		
+		// If the selection should be synchronized with the editor's selection, select any actors that
+		// are currently selected in the editor
+		if (const UCompositeEditorPanelSettings* Settings = GetDefault<UCompositeEditorPanelSettings>())
+		{
+			if (Settings->bSyncActorSelection)
+			{
+				if (USelection* Selection = GEditor->GetSelectedActors())
+				{
+					TArray<AActor*> SelectedActors;	
+					Selection->GetSelectedObjects<AActor>(SelectedActors);
+			
+					ListView->ClearSelection();
+					for (const AActor* Actor : SelectedActors)
+					{
+						FActorListItemRef* ActorToSelect = FilteredActorListItems.FindByPredicate([&Actor](const FActorListItemRef& Item)
+						{
+							return Item.IsValid() && Item->Actor.IsValid() && Item->Actor.Get() == Actor;
+						});
+
+						if (ActorToSelect)
+						{
+							ListView->SetItemSelection(*ActorToSelect, true);
+						}
+					}
+				}
+			}
+		}
 	}
 }
 
@@ -667,12 +758,14 @@ TSharedPtr<SWidget> SCompositeActorPickerTable::CreateListContextMenu()
 		{
 			MenuBuilder.BeginSection("ApplyMaterial", LOCTEXT("ApplyMaterial", "Apply Material"));
 
+			const FString LitMaterialPath = TEXT("/Composite/Materials/M_CompositeMesh_Lit_Masked.M_CompositeMesh_Lit_Masked");
+			const FString UnlitMaterialPath = TEXT("/Composite/Materials/M_CompositeMesh_Unlit_AlphaComposite.M_CompositeMesh_Unlit_AlphaComposite");
 			MenuBuilder.AddMenuEntry(
 				LOCTEXT("ApplyLitMaskedMaterial", "Apply Lit Masked Material To Selected Actors"),
 				LOCTEXT("ApplyLitMaskedMaterialToolTip", "Apply the plugin default lit masked material to selected actors: best for catching shadows and reflections."),
-				FSlateIcon(),
+				FSlateIcon(FCompositeEditorStyle::Get().GetStyleSetName(), "CompositeEditor.LitMaterial"),
 				FUIAction(
-					FExecuteAction::CreateSP(this, &SCompositeActorPickerTable::ApplyMaterialToActors, TEXT("/Composite/Materials/M_CompositeMesh_Lit_Masked.M_CompositeMesh_Lit_Masked")),
+					FExecuteAction::CreateSP(this, &SCompositeActorPickerTable::ApplyMaterialToActors, LitMaterialPath),
 					FCanExecuteAction::CreateSP(this, &SCompositeActorPickerTable::CanEditActors)
 				)
 			);
@@ -680,12 +773,34 @@ TSharedPtr<SWidget> SCompositeActorPickerTable::CreateListContextMenu()
 			MenuBuilder.AddMenuEntry(
 				LOCTEXT("ApplyUnlitAlphaMaterial", "Apply Unlit Alpha Material To Selected Actors"),
 				LOCTEXT("ApplyUnlitAlphaMaterialToolTip", "Apply the plugin default unlit alpha material to selected actors: best for keying media."),
-				FSlateIcon(),
+				FSlateIcon(FCompositeEditorStyle::Get().GetStyleSetName(), "CompositeEditor.UnlitMaterial"),
 				FUIAction(
-					FExecuteAction::CreateSP(this, &SCompositeActorPickerTable::ApplyMaterialToActors, TEXT("/Composite/Materials/M_CompositeMesh_Unlit_AlphaComposite.M_CompositeMesh_Unlit_AlphaComposite")),
+					FExecuteAction::CreateSP(this, &SCompositeActorPickerTable::ApplyMaterialToActors, UnlitMaterialPath),
 					FCanExecuteAction::CreateSP(this, &SCompositeActorPickerTable::CanEditActors)
 				)
 			);
+
+			MenuBuilder.AddMenuEntry(
+				LOCTEXT("ApplyCustomMaterial", "Apply Custom Material To Selected Actors"),
+				LOCTEXT("ApplyCustomMaterialToolTip", "Apply a custom material to selected actors"),
+				FSlateIcon(FCompositeEditorStyle::Get().GetStyleSetName(), "CompositeEditor.CustomMaterial"),
+				FUIAction(
+					FExecuteAction::CreateSP(this, &SCompositeActorPickerTable::ApplyMaterialToActors, FString()),
+					FCanExecuteAction::CreateSP(this, &SCompositeActorPickerTable::CanApplyCustomMaterial)
+				)
+			);
+
+			MenuBuilder.AddSeparator();
+
+			MenuBuilder.AddSubMenu(
+				TAttribute<FText>::CreateSP(this, &SCompositeActorPickerTable::GetCustomMaterialSubMenuLabel),
+				FText::GetEmpty(),
+				FNewMenuDelegate::CreateSP(this, &SCompositeActorPickerTable::CreateCustomMaterialSubMenu),
+				false,
+				FSlateIcon(FCompositeEditorStyle::Get().GetStyleSetName(), "CompositeEditor.CustomMaterial"),
+				false
+			);
+
 			MenuBuilder.EndSection();
 		}
 		
@@ -711,6 +826,76 @@ TSharedPtr<SWidget> SCompositeActorPickerTable::CreateListContextMenu()
 	return SNullWidget::NullWidget;
 }
 
+void SCompositeActorPickerTable::OnActorSelectionChanged(FActorListItemRef ActorListItem, ESelectInfo::Type InSelectType)
+{
+	if (InSelectType == ESelectInfo::Direct)
+	{
+		return;
+	}
+	
+	if (const UCompositeEditorPanelSettings* Settings = GetDefault<UCompositeEditorPanelSettings>())
+	{
+		if (Settings->bSyncActorSelection)
+		{
+			TArray<FActorListItemRef> SelectedActors = ListView->GetSelectedItems();
+			
+			GEditor->GetSelectedActors()->Modify();
+
+			GEditor->GetSelectedActors()->BeginBatchSelectOperation();
+			GEditor->GetSelectedActors()->DeselectAll();
+
+			for (const FActorListItemRef& Actor : SelectedActors)
+			{
+				constexpr bool bIsSelected = true;
+				constexpr bool bNotifyAfterSelect = false;
+				GEditor->SelectActor(Actor->Actor.Get(), bIsSelected, bNotifyAfterSelect);
+			}
+
+			constexpr bool bNotify = false;
+			GEditor->GetSelectedActors()->EndBatchSelectOperation(bNotify);
+
+			GEditor->NoteSelectionChange();
+		}
+	}
+}
+
+void SCompositeActorPickerTable::OnLevelSelectionChanged(UObject* NewSelection)
+{
+	USelection* Selection = Cast<USelection>(NewSelection);
+	if (!Selection)
+	{
+		return;
+	}
+
+	if (!ListView.IsValid())
+	{
+		return;
+	}
+
+	if (const UCompositeEditorPanelSettings* Settings = GetDefault<UCompositeEditorPanelSettings>())
+	{
+		if (Settings->bSyncActorSelection)
+		{
+			TArray<AActor*> SelectedActors;	
+			Selection->GetSelectedObjects<AActor>(SelectedActors);
+			
+			ListView->ClearSelection();
+			for (const AActor* Actor : SelectedActors)
+			{
+				FActorListItemRef* ActorToSelect = FilteredActorListItems.FindByPredicate([&Actor](const FActorListItemRef& Item)
+				{
+					return Item.IsValid() && Item->Actor.IsValid() && Item->Actor.Get() == Actor;
+				});
+
+				if (ActorToSelect)
+				{
+					ListView->SetItemSelection(*ActorToSelect, true);
+				}
+			}
+		}
+	}
+}
+
 bool SCompositeActorPickerTable::OnAllowListDrop(TSharedPtr<FDragDropOperation> InDragDropOperation)
 {
 	// Support dragging both actors and folders from the Outliner (dragging a folder will add all actor types in the folder)
@@ -744,6 +929,7 @@ FReply SCompositeActorPickerTable::OnListDropped(const FGeometry& Geometry, cons
 			
 	TSharedPtr<FScopedTransaction> AddActorsTransaction = nullptr;
 	int32 LastModifiedIndex = INDEX_NONE;
+	TArray<TSoftObjectPtr<AActor>> AddedActors;
 	for (AActor* DroppedActor : DroppedActors)
 	{
 		if (ActorListRef.ActorList->Contains(DroppedActor))
@@ -760,12 +946,19 @@ FReply SCompositeActorPickerTable::OnListDropped(const FGeometry& Geometry, cons
 		}
 
 		ActorListRef.ActorList->Add(DroppedActor);
+		AddedActors.Add(DroppedActor);
 		LastModifiedIndex = ActorListRef.ActorList->Num() - 1;
 	}
 
 	if (AddActorsTransaction.IsValid())
 	{
+		if (!AddedActors.IsEmpty())
+		{
+			ActorListRef.OnActorsAdded.ExecuteIfBound(AddedActors);
+		}
+
 		ActorListRef.NotifyPostEditChangeList(EPropertyChangeType::ArrayAdd, LastModifiedIndex);
+
 		FillActorList();
 	}
 
@@ -881,10 +1074,23 @@ void SCompositeActorPickerTable::RemoveActors()
 	FillActorList();
 }
 
-void SCompositeActorPickerTable::ApplyMaterialToActors(const TCHAR* InMaterialPath)
+void SCompositeActorPickerTable::ApplyMaterialToActors(FString InMaterialPath)
 {
-	UMaterial* LoadedMaterial = LoadObject<UMaterial>(nullptr, InMaterialPath);
-	if (!ensureMsgf(LoadedMaterial, TEXT("Expected valid material to apply to actors.")))
+	UMaterialInterface* LoadedMaterial = nullptr;
+
+	if (InMaterialPath.IsEmpty())
+	{
+		if (CustomMaterialAsset.IsValid() && CustomMaterialAsset.IsInstanceOf(UMaterialInterface::StaticClass()))
+		{
+			LoadedMaterial = Cast<UMaterialInterface>(CustomMaterialAsset.GetAsset());
+		}
+	}
+	else
+	{
+		LoadedMaterial = LoadObject<UMaterialInterface>(nullptr, InMaterialPath);
+	}
+
+	if (LoadedMaterial == nullptr)
 	{
 		return;
 	}
@@ -972,6 +1178,51 @@ bool SCompositeActorPickerTable::CanEditActors() const
 	return true;
 }
 
+bool SCompositeActorPickerTable::CanApplyCustomMaterial() const
+{
+	return CanEditActors() && CustomMaterialAsset.IsValid() && CustomMaterialAsset.IsInstanceOf(UMaterialInterface::StaticClass());
+}
+
+FText SCompositeActorPickerTable::GetCustomMaterialSubMenuLabel() const
+{
+	if (CustomMaterialAsset.IsValid() && CustomMaterialAsset.IsInstanceOf(UMaterialInterface::StaticClass()))
+	{
+		return FText::FromName(CustomMaterialAsset.AssetName);
+	}
+	
+	return LOCTEXT("SelectCustomMaterial", "Select a Custom Material");
+}
+
+void SCompositeActorPickerTable::CreateCustomMaterialSubMenu(FMenuBuilder& InMenuBuilder)
+{
+	FAssetPickerConfig AssetPickerConfig;
+	AssetPickerConfig.Filter.ClassPaths.Add(UMaterial::StaticClass()->GetClassPathName());
+	AssetPickerConfig.Filter.ClassPaths.Add(UMaterialInstance::StaticClass()->GetClassPathName());
+	AssetPickerConfig.Filter.bRecursiveClasses = true;
+	AssetPickerConfig.bAllowNullSelection = false;
+	AssetPickerConfig.OnAssetSelected = FOnAssetSelected::CreateSP(this, &SCompositeActorPickerTable::OnCustomMaterialSelected);
+	AssetPickerConfig.InitialAssetViewType = EAssetViewType::List;
+	AssetPickerConfig.InitialAssetSelection = CustomMaterialAsset;
+
+	FContentBrowserModule& ContentBrowserModule = FModuleManager::Get().LoadModuleChecked<FContentBrowserModule>(TEXT("ContentBrowser"));
+
+	InMenuBuilder.AddWidget(
+		SNew(SBox)
+		.MaxDesiredHeight(600.0f)
+		.WidthOverride(256.0f)
+		[
+			ContentBrowserModule.Get().CreateAssetPicker(AssetPickerConfig)
+		], 
+		FText(),
+		true
+	);
+}
+
+void SCompositeActorPickerTable::OnCustomMaterialSelected(const FAssetData& InAssetData)
+{
+	CustomMaterialAsset = InAssetData;
+}
+
 void SCompositeActorPickerTable::OnObjectPropertyChanged(UObject* InObject, FPropertyChangedEvent& InPropertyChangedEvent)
 {
 	if (!ActorListRef.IsValid())
diff --git a/Engine/Plugins/Compositing/Composite/Source/CompositeEditor/Private/UI/SCompositeActorPickerTable.h b/Engine/Plugins/Compositing/Composite/Source/CompositeEditor/Private/UI/SCompositeActorPickerTable.h
--- a/Engine/Plugins/Compositing/Composite/Source/CompositeEditor/Private/UI/SCompositeActorPickerTable.h
+++ b/Engine/Plugins/Compositing/Composite/Source/CompositeEditor/Private/UI/SCompositeActorPickerTable.h
@@ -3,10 +3,12 @@
 #pragma once
 
 #include "EditorUndoClient.h"
+#include "AssetRegistry/AssetData.h"
 #include "UObject/SoftObjectPtr.h"
 #include "UObject/UnrealType.h"
 #include "Widgets/SCompoundWidget.h"
 
+class UMaterialInterface;
 class UMaterial;
 struct FSceneOutlinerFilters;
 class SCompositeActorPickerToolbar;
@@ -18,6 +20,8 @@ class SHeaderRow;
 /** Stores a pointer reference to a list of actors on a UObject to be managed by the actor picker */
 struct FCompositeActorPickerListRef
 {
+	DECLARE_DELEGATE_OneParam(FOnActorsAdded, TArray<TSoftObjectPtr<AActor>>&);
+
 public:
 	FCompositeActorPickerListRef(const TWeakObjectPtr<UObject>& InActorListOwner, const FName& InActorListPropertyName, TArray<TSoftObjectPtr<AActor>>* InActorList);
 	
@@ -47,6 +51,9 @@ public:
 	
 	/** Pointer to the actor list being edited by this widget */
 	TArray<TSoftObjectPtr<AActor>>* ActorList = nullptr;
+
+	/** Raised when any actors are added to the actor list */
+	FOnActorsAdded OnActorsAdded;
 };
 
 /**
@@ -75,7 +82,7 @@ public:
 	DECLARE_DELEGATE_RetVal_TwoParams(TSharedRef<SWidget>, FOnGenerateColumnWidget, const TSoftObjectPtr<AActor>& /* InActor */, const FName& /* ColumnName */)
 	DECLARE_DELEGATE_OneParam(FOnExtendAddMenu, FMenuBuilder&);
 	DECLARE_DELEGATE_TwoParams(FOnExtendContextMenu, FMenuBuilder&, TArray<TSoftObjectPtr<AActor>>& /*SelectedActors*/);
-	DECLARE_DELEGATE_TwoParams(FOnApplyMaterialToActor, TSoftObjectPtr<AActor>& /* InActor */, UMaterial* /* InMaterial */);
+	DECLARE_DELEGATE_TwoParams(FOnApplyMaterialToActor, TSoftObjectPtr<AActor>& /* InActor */, UMaterialInterface* /* InMaterial */);
 
 public:
 	/** Returns default scene outliner filters that check IsAllowedPrimitiveClass and optionally exclude ACompositeMeshActor. */
@@ -83,6 +90,8 @@ public:
 
 	SLATE_BEGIN_ARGS(SCompositeActorPickerTable) : _ShowApplyMaterialSection(false) { }
 		SLATE_ATTRIBUTE(TSharedPtr<FSceneOutlinerFilters>, SceneOutlinerFilters)
+		SLATE_ARGUMENT(TArray<FName>, HiddenColumnsList)
+		SLATE_EVENT(FSimpleDelegate, OnHiddenColumnsListChanged)
 		SLATE_EVENT(FOnGenerateHeaderColumns, OnGenerateHeaderColumns)
 		SLATE_EVENT(FOnGenerateColumnWidget, OnGenerateColumnWidget)
 		SLATE_EVENT(FOnExtendAddMenu, OnExtendAddMenu)
@@ -105,7 +114,10 @@ public:
 	virtual void PostUndo(bool bSuccess) override;
 	virtual void PostRedo(bool bSuccess) override;
 	// End of FEditorUndoClient
-	
+
+	/** Gets the list of columns that are currently hidden in the actor table */
+	TArray<FName> GetHiddenColumnsList() const;
+
 private:
 	/** Binds any commands to callbacks for the actor list view */
 	void BindCommands();
@@ -122,6 +134,12 @@ private:
 	/** Create the right click context menu for the actor list view */
 	TSharedPtr<SWidget> CreateListContextMenu();
 
+	/** Raised when an actor item is selected in the list view */
+	void OnActorSelectionChanged(FActorListItemRef ActorListItem, ESelectInfo::Type InSelectType);
+
+	/** Raised when the editor selection is changed */
+	void OnLevelSelectionChanged(UObject* NewSelection);
+
 	/** Raised when a drag drop operation is trying to be dropped on the actor list */
 	bool OnAllowListDrop(TSharedPtr<FDragDropOperation> InDragDropOperation);
 
@@ -138,11 +156,23 @@ private:
 	void RemoveActors();
 
 	/** Apply the specified material to the the selected actors from the list */
-	void ApplyMaterialToActors(const TCHAR* InMaterialPath);
+	void ApplyMaterialToActors(FString InMaterialPath);
 
 	/** Gets whether the contextual menu commands can be used */
 	bool CanEditActors() const;
 	
+	/** Gets whether a custom material can be applied to the selected actors */
+	bool CanApplyCustomMaterial() const;
+
+	/** Gets the label for the custom material sub-menu in the context menu */
+	FText GetCustomMaterialSubMenuLabel() const;
+
+	/** Creates the custom material sub-menu in the context menu */
+	void CreateCustomMaterialSubMenu(FMenuBuilder& InMenuBuilder);
+
+	/** Raised when a new material is selected in the material browser in the custom material sub-menu */
+	void OnCustomMaterialSelected(const FAssetData& InAssetData);
+	
 	/** Raised when a property on the object has changed */
 	void OnObjectPropertyChanged(UObject* InObject, FPropertyChangedEvent& InPropertyChangedEvent);
 	
@@ -150,6 +180,9 @@ private:
 	/** List view that displays the list of selected actors */
 	TSharedPtr<SListView<FActorListItemRef>> ListView;
 
+	/** Header row of the list view that displays the columns of the list */
+	TSharedPtr<SHeaderRow> HeaderRow;
+
 	/** The toolbar that contains the Add button and the filter text box */
 	TSharedPtr<SCompositeActorPickerToolbar> Toolbar;
 	
@@ -180,6 +213,9 @@ private:
 	/** Flag to enable the material application section for selected actors. */
 	TAttribute<bool> bShowApplyMaterialSection;
 
+	/** Selected custom material to apply to selected actors */
+	FAssetData CustomMaterialAsset = FAssetData();
+
 	/** Raised when a material is being applied to an actor. Allows for custom logic on how the material is applied */
 	FOnApplyMaterialToActor OnApplyMaterialToActor;
 
