Composure: Dilation pass alpha threshold clamp max adjustment.

#Synced-CL 52734496

diff --git a/Engine/Plugins/Compositing/Composite/Source/Composite/Public/Passes/CompositePassDilation.h b/Engine/Plugins/Compositing/Composite/Source/Composite/Public/Passes/CompositePassDilation.h
--- a/Engine/Plugins/Compositing/Composite/Source/Composite/Public/Passes/CompositePassDilation.h
+++ b/Engine/Plugins/Compositing/Composite/Source/Composite/Public/Passes/CompositePassDilation.h
@@ -44,7 +44,7 @@ public:
 	int32 Size = 1;
 
 	/** Opacity threshold below which pixels are treated as fully transparent. In standard opacity space (0=transparent, 1=opaque). */
-	UPROPERTY(EditAnywhere, BlueprintReadWrite, Category = "Composite", meta = (ClampMin = "0.0", ClampMax = "0.5"))
+	UPROPERTY(EditAnywhere, BlueprintReadWrite, Category = "Composite", meta = (ClampMin = "0.0", ClampMax = "1.0"))
 	float AlphaThreshold = 0.0001f;
 };
 
