Composure: Rename CenteredScale pass to Transform2D for 2D post-move transform support.

Port the legacy Composure 2D post-move transforms into the Composite plugin by extending the CenteredScale pass into a full Transform2D pass. Adds Translation, RotationAngle, and Pivot properties alongside the existing scale modes, with an aspect-ratio-corrected affine transform in the shader matching the legacy GetCroppingUVTransformationMatrix conventions. All properties are Interp-enabled for Sequencer animation.

#Synced-CL 52740614

diff --git a/Engine/Plugins/Compositing/Composite/Config/DefaultComposite.ini b/Engine/Plugins/Compositing/Composite/Config/DefaultComposite.ini
--- a/Engine/Plugins/Compositing/Composite/Config/DefaultComposite.ini
+++ b/Engine/Plugins/Compositing/Composite/Config/DefaultComposite.ini
@@ -45,3 +45,4 @@
 +PropertyRedirects=(OldName="/Script/Composite.CompositePassCenteredScale.TargetAspectRatio", NewName="/Script/Composite.CompositePassCenteredScale.ContentAspectRatio")
 +PropertyRedirects=(OldName="/Script/Composite.CompositePassCenteredScale.ScaleFactor", NewName="/Script/Composite.CompositePassCenteredScale.ManualScale")
 +ClassRedirects=(OldName="/Script/Composite.CompositePassAlphaTransform",NewName="/Script/Composite.CompositePassTranslucency")
++ClassRedirects=(OldName="/Script/Composite.CompositePassCenteredScale",NewName="/Script/Composite.CompositePassTransform2D")
diff --git a/Engine/Plugins/Compositing/Composite/Shaders/Private/CompositeCenteredScale.usf b/Engine/Plugins/Compositing/Composite/Shaders/Private/CompositeCenteredScale.usf
deleted file mode 100644
diff --git a/Engine/Plugins/Compositing/Composite/Shaders/Private/CompositeTransform2D.usf b/Engine/Plugins/Compositing/Composite/Shaders/Private/CompositeTransform2D.usf
new file mode 100644
--- /dev/null
+++ b/Engine/Plugins/Compositing/Composite/Shaders/Private/CompositeTransform2D.usf
@@ -0,0 +1,45 @@
+// Copyright Epic Games, Inc. All Rights Reserved.
+
+#include "/Engine/Private/Common.ush"
+#include "/Engine/Private/ScreenPass.ush"
+
+Texture2D Input_Texture;
+SamplerState Input_Sampler;
+FScreenTransform SvPositionToOutputViewportUV;
+FScreenTransform OutputViewportUVToInputTextureUV;
+float2 Pivot;
+float2 ScaleUV;
+float RotationRadians;
+float2 Translation;
+float AspectRatio;
+
+void MainPS(
+    float4 SvPosition : SV_POSITION,
+    out float4 OutColor : SV_Target0
+    )
+{
+    float2 OutputViewportUV = ApplyScreenTransform(SvPosition.xy, SvPositionToOutputViewportUV);
+
+    // 2D affine transform: output UV -> input UV (sampling direction).
+    // User-facing Y convention is Y-up (Translation.Y > 0 moves image up), but screen UV
+    // is Y-down, so flip the Y components of Translation and Pivot for sampling.
+    float2 UVPivot = float2(Pivot.x, 1.0 - Pivot.y);
+    float2 UV = OutputViewportUV;
+    UV -= float2(Translation.x, -Translation.y);
+    UV -= UVPivot;
+
+    // Aspect-correct rotation: enter isotropic space (equal physical distance per unit
+    // in both axes), rotate, then leave isotropic space. Dividing by AR before rotation
+    // ensures perpendicular vectors remain perpendicular in screen pixel space.
+    UV.y /= AspectRatio;
+    float c = cos(RotationRadians);
+    float s = sin(RotationRadians);
+    UV = float2(UV.x * c - UV.y * s,
+                UV.x * s + UV.y * c);
+    UV *= ScaleUV;
+    UV.y *= AspectRatio;
+    UV += UVPivot;
+
+    float2 InputUV = ApplyScreenTransform(UV, OutputViewportUVToInputTextureUV);
+    OutColor = Input_Texture.Sample(Input_Sampler, InputUV);
+}
diff --git a/Engine/Plugins/Compositing/Composite/Source/Composite/Private/Layers/CompositeLayerPlate.cpp b/Engine/Plugins/Compositing/Composite/Source/Composite/Private/Layers/CompositeLayerPlate.cpp
--- a/Engine/Plugins/Compositing/Composite/Source/Composite/Private/Layers/CompositeLayerPlate.cpp
+++ b/Engine/Plugins/Compositing/Composite/Source/Composite/Private/Layers/CompositeLayerPlate.cpp
@@ -18,7 +18,7 @@
 #include "Materials/MaterialInstanceDynamic.h"
 #include "Passes/CompositePassDistortion.h"
 #include "Passes/CompositePassTranslucency.h"
-#include "Passes/CompositePassCenteredScale.h"
+#include "Passes/CompositePassTransform2D.h"
 #include "Profile/IMediaProfileManager.h"
 #include "Profile/MediaProfile.h"
 #include "Profile/MediaProfilePlaybackManager.h"
@@ -184,7 +184,7 @@ UCompositeLayerPlate::UCompositeLayerPlate(const FObjectInitializer& ObjectIniti
 	ScenePasses.Add(CreateDefaultSubobject<UCompositePassDistortion>(TEXT("CompositeUndistortScenePass")));
 
 	// Internal pass to automatically remove overscan
-	OverscanPass = CreateDefaultSubobject<UCompositePassCenteredScale>(TEXT("CompositeOverscanPass"));
+	OverscanPass = CreateDefaultSubobject<UCompositePassTransform2D>(TEXT("CompositeOverscanPass"));
 	OverscanPass->ScaleMode = ECompositePassScaleMode::None;
 	OverscanPass->bRemoveOverscan = true;
 
diff --git a/Engine/Plugins/Compositing/Composite/Source/Composite/Private/Passes/CompositePassCenteredScale.cpp b/Engine/Plugins/Compositing/Composite/Source/Composite/Private/Passes/CompositePassCenteredScale.cpp
deleted file mode 100644
diff --git a/Engine/Plugins/Compositing/Composite/Source/Composite/Private/Passes/CompositePassTransform2D.cpp b/Engine/Plugins/Compositing/Composite/Source/Composite/Private/Passes/CompositePassTransform2D.cpp
new file mode 100644
--- /dev/null
+++ b/Engine/Plugins/Compositing/Composite/Source/Composite/Private/Passes/CompositePassTransform2D.cpp
@@ -0,0 +1,257 @@
+// Copyright Epic Games, Inc. All Rights Reserved.
+
+#include "Passes/CompositePassTransform2D.h"
+
+#include "Camera/CameraComponent.h"
+#include "CompositeActor.h"
+#include "Engine/TextureRenderTarget2D.h"
+#include "Engine/Texture2D.h"
+#include "Engine/Texture2DDynamic.h"
+#include "Layers/CompositeLayerPlate.h"
+#include "MediaTexture.h"
+#include "Passes/CompositeCorePassProxy.h"
+#include "PixelShaderUtils.h"
+#include "ScreenPass.h"
+#include "SceneView.h"
+
+DECLARE_GPU_STAT_NAMED(FCompositeTransform2D, TEXT("Composite.Transform2D"));
+
+class FCompositePassTransform2DShader : public FGlobalShader
+{
+	DECLARE_GLOBAL_SHADER(FCompositePassTransform2DShader);
+	SHADER_USE_PARAMETER_STRUCT(FCompositePassTransform2DShader, FGlobalShader);
+
+	BEGIN_SHADER_PARAMETER_STRUCT(FParameters, )
+		SHADER_PARAMETER_STRUCT(FScreenPassTextureInput, Input)
+		SHADER_PARAMETER(FScreenTransform, SvPositionToOutputViewportUV)
+		SHADER_PARAMETER(FScreenTransform, OutputViewportUVToInputTextureUV)
+		SHADER_PARAMETER(FVector2f, Pivot)
+		SHADER_PARAMETER(FVector2f, ScaleUV)
+		SHADER_PARAMETER(float, RotationRadians)
+		SHADER_PARAMETER(FVector2f, Translation)
+		SHADER_PARAMETER(float, AspectRatio)
+		RENDER_TARGET_BINDING_SLOTS()
+	END_SHADER_PARAMETER_STRUCT()
+};
+IMPLEMENT_GLOBAL_SHADER(FCompositePassTransform2DShader, "/Plugin/Composite/Private/CompositeTransform2D.usf", "MainPS", SF_Pixel);
+
+
+namespace UE::Composite::Private
+{
+	using namespace CompositeCore;
+
+	FPassTexture FTransform2DPassProxy::Add(FRDGBuilder& GraphBuilder, const FSceneView& InView, const FPassInputArray& Inputs, const FPassContext& PassContext) const
+	{
+		RDG_EVENT_SCOPE_STAT(GraphBuilder, FCompositeTransform2D, "Composite.Transform2D");
+
+		check(ValidateInputs(Inputs));
+
+		const UE::CompositeCore::FResourceMetadata& Metadata = Inputs[0].Metadata;
+		const FScreenPassTexture& Input = Inputs[0].Texture;
+		FScreenPassRenderTarget Output = Inputs.OverrideOutput;
+
+		// No override output: create an intermediate render target.
+		if (!Output.IsValid())
+		{
+			Output = CreateOutputRenderTarget(GraphBuilder, InView, PassContext.OutputViewRect, Input.Texture->Desc, TEXT("Transform2DCompositePass"));
+		}
+
+		FGlobalShaderMap* GlobalShaderMap = GetGlobalShaderMap(InView.GetFeatureLevel());
+		FCompositePassTransform2DShader::FParameters* Parameters = GraphBuilder.AllocParameters<FCompositePassTransform2DShader::FParameters>();
+		FRHISamplerState* InputSampler = Metadata.GetSamplerState();
+		Parameters->Input = GetScreenPassTextureInput(Input, InputSampler);
+		Parameters->SvPositionToOutputViewportUV = FScreenTransform::ChangeTextureBasisFromTo(
+			FScreenPassTextureViewport(Output),
+			FScreenTransform::ETextureBasis::TexelPosition,
+			FScreenTransform::ETextureBasis::ViewportUV);
+		Parameters->OutputViewportUVToInputTextureUV = FScreenTransform::ChangeTextureBasisFromTo(
+			FScreenPassTextureViewport(Input),
+			FScreenTransform::ETextureBasis::ViewportUV,
+			FScreenTransform::ETextureBasis::TextureUV);
+		Parameters->ScaleUV = ScaleUV;
+		Parameters->Translation = Translation;
+		Parameters->RotationRadians = RotationRadians;
+		Parameters->Pivot = Pivot;
+		Parameters->AspectRatio = AspectRatio;
+		Parameters->RenderTargets[0] = Output.GetRenderTargetBinding();
+
+		TShaderMapRef<FCompositePassTransform2DShader> PixelShader(GlobalShaderMap);
+		FPixelShaderUtils::AddFullscreenPass(
+			GraphBuilder,
+			GlobalShaderMap,
+			RDG_EVENT_NAME("Composite.Transform2D (%dx%d)", Output.ViewRect.Width(), Output.ViewRect.Height()),
+			PixelShader,
+			Parameters,
+			Output.ViewRect
+		);
+
+		return UE::CompositeCore::FPassTexture{ MoveTemp(Output), Metadata };
+	}
+}
+
+UCompositePassTransform2D::UCompositePassTransform2D(const FObjectInitializer& ObjectInitializer)
+	: Super(ObjectInitializer)
+	, Pivot(0.5f, 0.5f)
+	, ScaleMode(ECompositePassScaleMode::Automatic)
+	, ContainerAspectRatio(16.0f, 9.0f)
+	, ContentAspectRatio(16.0f, 9.0f)
+	, bScaleSingleAxis(true)
+	, UnpadPixels(FIntPoint::ZeroValue)
+	, ManualScale(1.0f, 1.0f)
+	, bRemoveOverscan(true)
+	, RotationAngle(0.0f)
+	, Translation(FVector2f::ZeroVector)
+{
+}
+
+UCompositePassTransform2D::~UCompositePassTransform2D() = default;
+
+FCompositeCorePassProxy* UCompositePassTransform2D::GetProxy(const UE::CompositeCore::FPassInputDecl& InputDecl, FCompositeTraversalContext& InContext, FSceneRenderingBulkObjectAllocator& InFrameAllocator) const
+{
+	using namespace UE::Composite::Private;
+
+	// Use cached value from GetIsActive() if available, otherwise compute fresh.
+	const FVector2f ScaleUV = CachedScaleUV.IsSet() ? CachedScaleUV.GetValue() : CalculateScale();
+	CachedScaleUV.Reset();
+
+	FTransform2DPassProxy* Proxy = InFrameAllocator.Create<FTransform2DPassProxy>(FPassInputDeclArray{ InputDecl });
+	Proxy->ScaleUV = ScaleUV;
+	Proxy->Translation = Translation;
+	Proxy->RotationRadians = -FMath::DegreesToRadians(RotationAngle);
+	Proxy->Pivot = Pivot;
+
+	// Compute aspect ratio for rotation correction.
+	const FIntVector2 TextureSize = GetPlateTextureSize();
+	Proxy->AspectRatio = (TextureSize.X > 0 && TextureSize.Y > 0)
+		? static_cast<float>(TextureSize.X) / TextureSize.Y
+		: 1.0f;
+
+	return Proxy;
+}
+
+bool UCompositePassTransform2D::GetIsActive() const
+{
+	CachedScaleUV = CalculateScale();
+
+	const bool bHasScale = !CachedScaleUV.GetValue().Equals(FVector2f::One(), UE_SMALL_NUMBER);
+	const bool bHasRotation = !FMath::IsNearlyZero(RotationAngle, UE_SMALL_NUMBER);
+	const bool bHasTranslation = !Translation.IsNearlyZero(UE_SMALL_NUMBER);
+
+	return GetIsEnabled()
+		&& (ScaleMode != ECompositePassScaleMode::None || bRemoveOverscan || bHasRotation || bHasTranslation)
+		&& (bHasScale || bHasRotation || bHasTranslation);
+}
+
+FVector2f UCompositePassTransform2D::CalculateScale() const
+{
+	ACompositeActor* CompositeActor = GetTypedOuter<ACompositeActor>();
+	UCameraComponent* CameraComponent = nullptr;
+
+	if (CompositeActor)
+	{
+		CameraComponent = CompositeActor->GetCameraComponent();
+	}
+
+	FVector2f ScaleUV = FVector2f::One();
+
+	auto ApplyAspectRatioScaleFn = [this, &ScaleUV](float ContainerAR, float ContentAR)
+		{
+			if (ContainerAR > ContentAR)
+			{
+				const float Factor = (ContainerAR != 0.0f) ? ContentAR / ContainerAR : 0.0f;
+
+				ScaleUV *= FVector2f(Factor, bScaleSingleAxis ? 1.0f : Factor);
+			}
+			else if (ContainerAR < ContentAR)
+			{
+				const float Factor = (ContentAR != 0.0f) ? ContainerAR / ContentAR : 0.0f;
+
+				ScaleUV *= FVector2f(bScaleSingleAxis ? 1.0f : Factor, Factor);
+			}
+		};
+
+	// Fetch plate texture dimensions (used by Automatic and Unpad modes).
+	const FIntVector2 TextureSize = (ScaleMode == ECompositePassScaleMode::Automatic || ScaleMode == ECompositePassScaleMode::Unpad)
+		? GetPlateTextureSize()
+		: FIntVector2::ZeroValue;
+
+	switch (ScaleMode)
+	{
+	case ECompositePassScaleMode::None:
+		break;
+	case ECompositePassScaleMode::Automatic:
+	{
+		if (TextureSize.X > 0 && TextureSize.Y > 0)
+		{
+			const float ContainerAR = FMath::SafeDivide(TextureSize.X, static_cast<float>(TextureSize.Y));
+			const float ContentAR = CameraComponent ? CameraComponent->AspectRatio : ContainerAR;
+
+			ApplyAspectRatioScaleFn(ContainerAR, ContentAR);
+		}
+	}
+	break;
+	case ECompositePassScaleMode::AspectRatio:
+	{
+		const float ContainerAR = FMath::SafeDivide(ContainerAspectRatio.X, ContainerAspectRatio.Y);
+		const float ContentAR = FMath::SafeDivide(ContentAspectRatio.X, ContentAspectRatio.Y);
+
+		ApplyAspectRatioScaleFn(ContainerAR, ContentAR);
+	}
+	break;
+	case ECompositePassScaleMode::Unpad:
+	{
+		if (TextureSize.X > 2 * UnpadPixels.X && TextureSize.X > 0)
+		{
+			ScaleUV.X *= static_cast<float>(TextureSize.X - 2 * UnpadPixels.X) / TextureSize.X;
+		}
+		if (TextureSize.Y > 2 * UnpadPixels.Y && TextureSize.Y > 0)
+		{
+			ScaleUV.Y *= static_cast<float>(TextureSize.Y - 2 * UnpadPixels.Y) / TextureSize.Y;
+		}
+	}
+	break;
+	case ECompositePassScaleMode::Manual:
+		ScaleUV.X = FMath::SafeDivide(ScaleUV.X, ManualScale.X);
+		ScaleUV.Y = FMath::SafeDivide(ScaleUV.Y, ManualScale.Y);
+		break;
+
+	default:
+		checkNoEntry();
+		break;
+	};
+
+	if (bRemoveOverscan && CameraComponent && CameraComponent->bCropOverscan)
+	{
+		ScaleUV *= 1.0f + CameraComponent->Overscan;
+	}
+
+	return ScaleUV;
+}
+
+FIntVector2 UCompositePassTransform2D::GetPlateTextureSize() const
+{
+	if (const UCompositeLayerPlate* PlateLayer = GetTypedOuter<UCompositeLayerPlate>())
+	{
+		if (IsValid(PlateLayer->Texture))
+		{
+			if (const UMediaTexture* MediaTexture = Cast<UMediaTexture>(PlateLayer->Texture))
+			{
+				return FIntVector2(MediaTexture->GetWidth(), MediaTexture->GetHeight());
+			}
+			else if (const UTexture2D* Texture2D = Cast<UTexture2D>(PlateLayer->Texture))
+			{
+				return FIntVector2(Texture2D->GetSizeX(), Texture2D->GetSizeY());
+			}
+			else if (const UTexture2DDynamic* Texture2DDynamic = Cast<UTexture2DDynamic>(PlateLayer->Texture))
+			{
+				return FIntVector2(Texture2DDynamic->SizeX, Texture2DDynamic->SizeY);
+			}
+			else if (const UTextureRenderTarget2D* RenderTarget2D = Cast<UTextureRenderTarget2D>(PlateLayer->Texture))
+			{
+				return FIntVector2(RenderTarget2D->SizeX, RenderTarget2D->SizeY);
+			}
+		}
+	}
+
+	return FIntVector2::ZeroValue;
+}
diff --git a/Engine/Plugins/Compositing/Composite/Source/Composite/Private/Tests/CompositePassCenteredScaleTest.cpp b/Engine/Plugins/Compositing/Composite/Source/Composite/Private/Tests/CompositePassCenteredScaleTest.cpp
deleted file mode 100644
diff --git a/Engine/Plugins/Compositing/Composite/Source/Composite/Private/Tests/CompositePassTransform2DTest.cpp b/Engine/Plugins/Compositing/Composite/Source/Composite/Private/Tests/CompositePassTransform2DTest.cpp
new file mode 100644
--- /dev/null
+++ b/Engine/Plugins/Compositing/Composite/Source/Composite/Private/Tests/CompositePassTransform2DTest.cpp
@@ -0,0 +1,231 @@
+// Copyright Epic Games, Inc. All Rights Reserved.
+
+#include "CoreTypes.h"
+#include "RenderGraphBuilder.h"
+#include "RHI.h"
+#include "ScreenPass.h"
+
+#include "Passes/CompositePassTransform2D.h"
+#include "Tests/CompositeTestHelpers.h"
+
+#if WITH_DEV_AUTOMATION_TESTS
+#include "Misc/AutomationTest.h"
+#endif
+
+namespace UE::Composite::Tests
+{
+
+#if WITH_DEV_AUTOMATION_TESTS
+
+// -- Upload/download parameter structs --
+
+BEGIN_SHADER_PARAMETER_STRUCT(FTransform2DTestUploadParameters, )
+	RDG_TEXTURE_ACCESS(InputTexture, ERHIAccess::CopyDest)
+END_SHADER_PARAMETER_STRUCT()
+
+BEGIN_SHADER_PARAMETER_STRUCT(FTransform2DTestDownloadParameters, )
+	RDG_TEXTURE_ACCESS(OutputTexture, ERHIAccess::CopySrc)
+END_SHADER_PARAMETER_STRUCT()
+
+namespace Transform2D  // unique name prevents symbol collision in unity builds
+{
+
+// -- CPU reference --
+
+// Input: horizontal gradient (Red = col/(W-1), G=B=0, A=1).
+// Computes the 2D affine transform matching the shader (output UV -> input UV), then bilinear-samples the gradient.
+static float Transform2DExpectedR(FVector2f ScaleUV, FVector2f Translation, float RotationDegrees,
+	FVector2f Pivot, float AspectRatio, int32 SampleCol, int32 SampleRow, int32 W, int32 H)
+{
+	const float UOut = (SampleCol + 0.5f) / W;
+	const float VOut = (SampleRow + 0.5f) / H;
+
+	const float RotRad = -FMath::DegreesToRadians(RotationDegrees);
+
+	// Y-up user convention -> Y-down UV space: flip Translation.Y and Pivot.Y.
+	const float UVPivotX = Pivot.X;
+	const float UVPivotY = 1.0f - Pivot.Y;
+
+	float U = UOut - Translation.X;
+	float V = VOut + Translation.Y;
+	U -= UVPivotX;
+	V -= UVPivotY;
+
+	V /= AspectRatio;
+	const float c = FMath::Cos(RotRad);
+	const float s = FMath::Sin(RotRad);
+	const float RotU = U * c - V * s;
+	const float RotV = U * s + V * c;
+	U = RotU;
+	V = RotV;
+	U *= ScaleUV.X;
+	V *= ScaleUV.Y;
+	V *= AspectRatio;
+	U += UVPivotX;
+	V += UVPivotY;
+
+	// Bilinear on a horizontal linear gradient: red = clamp(u*W - 0.5, 0, W-1) / (W-1).
+	return FMath::Clamp(U * W - 0.5f, 0.0f, float(W - 1)) / (W - 1);
+}
+
+// -- Test scenarios --
+
+struct FTransform2DScenario
+{
+	const TCHAR* Label;
+	FVector2f    ScaleUV;
+	FVector2f    Translation;
+	float        RotationDegrees;
+	FVector2f    Pivot;
+	float        AspectRatio;
+};
+
+// Input: horizontal gradient (Red = col/(W-1), G=B=0, A=1).
+// Sampled at column W/4, row H/2 - offset from center so transforms shift the sample point.
+static const FTransform2DScenario GScenarios[] =
+{
+	// Scale-only scenarios (backwards compatible with original CenteredScale tests)
+	{ TEXT("Identity"),        FVector2f(1.0f,  1.0f),  FVector2f(0.0f, 0.0f), 0.0f,  FVector2f(0.5f, 0.5f), 1.0f },
+	{ TEXT("ZoomIn"),          FVector2f(0.5f,  0.5f),  FVector2f(0.0f, 0.0f), 0.0f,  FVector2f(0.5f, 0.5f), 1.0f },
+	{ TEXT("ZoomOut"),         FVector2f(2.0f,  2.0f),  FVector2f(0.0f, 0.0f), 0.0f,  FVector2f(0.5f, 0.5f), 1.0f },
+	{ TEXT("HorizontalOnly"), FVector2f(0.75f, 1.0f),  FVector2f(0.0f, 0.0f), 0.0f,  FVector2f(0.5f, 0.5f), 1.0f },
+	// Translation
+	{ TEXT("TranslateX"),      FVector2f(1.0f, 1.0f),  FVector2f(0.1f, 0.0f), 0.0f,  FVector2f(0.5f, 0.5f), 1.0f },
+	// Rotation (square aspect ratio for predictable results)
+	{ TEXT("Rotate90"),        FVector2f(1.0f, 1.0f),  FVector2f(0.0f, 0.0f), 90.0f, FVector2f(0.5f, 0.5f), 1.0f },
+	// Custom pivot with scale
+	{ TEXT("CustomPivot"),     FVector2f(0.5f, 0.5f),  FVector2f(0.0f, 0.0f), 0.0f,  FVector2f(0.25f, 0.5f), 1.0f },
+	// Rotation with non-square aspect ratio
+	{ TEXT("Rotate45_Wide"),   FVector2f(1.0f, 1.0f),  FVector2f(0.0f, 0.0f), 45.0f, FVector2f(0.5f, 0.5f), 16.0f / 9.0f },
+	// Combined transform
+	{ TEXT("Combined"),        FVector2f(0.5f, 0.5f),  FVector2f(0.1f, 0.0f), 45.0f, FVector2f(0.5f, 0.5f), 1.0f },
+};
+
+}  // namespace Transform2D
+
+IMPLEMENT_SIMPLE_AUTOMATION_TEST(FCompositeTransform2DPassTest, "Composite.UnitTests.Transform2DPass",
+	EAutomationTestFlags_ApplicationContextMask | EAutomationTestFlags::EngineFilter | EAutomationTestFlags::NonNullRHI)
+
+bool FCompositeTransform2DPassTest::RunTest(const FString& Parameter)
+{
+	using namespace UE::Composite::Private;
+	using namespace UE::CompositeCore;
+	using namespace Transform2D;
+
+	const int32 NumScenarios = UE_ARRAY_COUNT(GScenarios);
+
+	TArray<FLinearColor> ResultRDG;
+	ResultRDG.SetNumZeroed(NumScenarios);
+
+	TArray<FTransform2DPassProxy> Proxies;
+	Proxies.Reserve(NumScenarios);
+	for (int32 i = 0; i < NumScenarios; ++i)
+	{
+		const FTransform2DScenario& S = GScenarios[i];
+		FTransform2DPassProxy Proxy(GetDefaultInputDeclArray());
+		Proxy.ScaleUV = S.ScaleUV;
+		Proxy.Translation = S.Translation;
+		Proxy.RotationRadians = -FMath::DegreesToRadians(S.RotationDegrees);
+		Proxy.Pivot = S.Pivot;
+		Proxy.AspectRatio = S.AspectRatio;
+		Proxies.Emplace(MoveTemp(Proxy));
+	}
+
+	FCompositeTestView TestView;
+	DispatchAndWait(TestView.Get(),
+		[&ResultRDG, NumScenarios, Proxies = MoveTemp(Proxies)](FRDGBuilder& GraphBuilder, const FSceneView& View) mutable
+		{
+			const FIntRect PixelRect(0, 0, GTestViewSize.X, GTestViewSize.Y);
+
+			for (int32 i = 0; i < NumScenarios; ++i)
+			{
+				const FRDGTextureDesc InputDesc = FRDGTextureDesc::Create2D(
+					GTestViewSize, PF_A32B32G32R32F,
+					FClearValueBinding::Black, ETextureCreateFlags::None);
+				FRDGTextureRef InputTexture = GraphBuilder.CreateTexture(InputDesc, TEXT("Transform2DTest.Input"));
+
+				const FRDGTextureDesc OutputDesc = FRDGTextureDesc::Create2D(
+					GTestViewSize, PF_A32B32G32R32F,
+					FClearValueBinding::Black, ETextureCreateFlags::RenderTargetable);
+				FRDGTextureRef OutputTexture = GraphBuilder.CreateTexture(OutputDesc, TEXT("Transform2DTest.Output"));
+
+				{
+					FTransform2DTestUploadParameters* Params = GraphBuilder.AllocParameters<FTransform2DTestUploadParameters>();
+					Params->InputTexture = InputTexture;
+
+					GraphBuilder.AddPass(
+						RDG_EVENT_NAME("Transform2DTest.Upload[%d]", i), Params,
+						ERDGPassFlags::Readback,
+						[Params](FRHICommandListImmediate& RHICmdList)
+						{
+							uint32 Stride = 0;
+							void* Data = RHICmdList.LockTexture2D(Params->InputTexture->GetRHI(), 0, RLM_WriteOnly, Stride, false);
+							for (int32 Row = 0; Row < GTestViewSize.Y; ++Row)
+							{
+								FLinearColor* RowPtr = reinterpret_cast<FLinearColor*>(static_cast<uint8*>(Data) + Row * Stride);
+								for (int32 Col = 0; Col < GTestViewSize.X; ++Col)
+								{
+									RowPtr[Col] = FLinearColor(Col / (GTestViewSize.X - 1.0f), 0.0f, 0.0f, 1.0f);
+								}
+							}
+							RHICmdList.UnlockTexture2D(Params->InputTexture->GetRHI(), 0, false);
+						});
+				}
+
+				FPassInputArray Inputs;
+				{
+					FPassInput& In = Inputs.GetArray().AddDefaulted_GetRef();
+					In.Texture         = FScreenPassTexture{ InputTexture, PixelRect };
+					In.Metadata.Filter = SF_Bilinear;
+				}
+				Inputs.OverrideOutput = FScreenPassRenderTarget{
+					FScreenPassTexture{ OutputTexture, PixelRect }, ERenderTargetLoadAction::ENoAction };
+
+				FPassContext PassContext;
+				PassContext.OutputViewRect = PixelRect;
+
+				Proxies[i].Add(GraphBuilder, View, Inputs, PassContext);
+
+				{
+					FTransform2DTestDownloadParameters* Params = GraphBuilder.AllocParameters<FTransform2DTestDownloadParameters>();
+					Params->OutputTexture = OutputTexture;
+
+					FLinearColor* ResultPtr = &ResultRDG[i];
+					GraphBuilder.AddPass(
+						RDG_EVENT_NAME("Transform2DTest.Download[%d]", i), Params,
+						ERDGPassFlags::Readback,
+						[ResultPtr, Params](FRHICommandListImmediate& RHICmdList)
+						{
+							uint32 Stride = 0;
+							const void* Data = RHICmdList.LockTexture2D(Params->OutputTexture->GetRHI(), 0, RLM_ReadOnly, Stride, false);
+							const FLinearColor* RowPtr = reinterpret_cast<const FLinearColor*>(
+								static_cast<const uint8*>(Data) + (GTestViewSize.Y / 2) * Stride);
+							FMemory::Memcpy(ResultPtr, &RowPtr[GTestViewSize.X / 4], sizeof(FLinearColor));
+							RHICmdList.UnlockTexture2D(Params->OutputTexture->GetRHI(), 0, false);
+						});
+				}
+			}
+		});
+
+	constexpr float Tolerance = 1e-3f;
+	const int32 SampleCol = GTestViewSize.X / 4;
+	const int32 SampleRow = GTestViewSize.Y / 2;
+	for (int32 i = 0; i < NumScenarios; ++i)
+	{
+		const FTransform2DScenario& S  = GScenarios[i];
+		const FLinearColor&         Actual = ResultRDG[i];
+
+		const float ExpectedR = Transform2DExpectedR(
+			S.ScaleUV, S.Translation, S.RotationDegrees, S.Pivot, S.AspectRatio,
+			SampleCol, SampleRow, GTestViewSize.X, GTestViewSize.Y);
+
+		UTEST_EQUAL_TOLERANCE(S.Label, Actual.R, ExpectedR, Tolerance);
+		UTEST_EQUAL_TOLERANCE(S.Label, Actual.A, 1.0f, Tolerance);
+	}
+
+	return true;
+}
+
+#endif // WITH_DEV_AUTOMATION_TESTS
+
+} // namespace UE::Composite::Tests
diff --git a/Engine/Plugins/Compositing/Composite/Source/Composite/Public/Layers/CompositeLayerPlate.h b/Engine/Plugins/Compositing/Composite/Source/Composite/Public/Layers/CompositeLayerPlate.h
--- a/Engine/Plugins/Compositing/Composite/Source/Composite/Public/Layers/CompositeLayerPlate.h
+++ b/Engine/Plugins/Compositing/Composite/Source/Composite/Public/Layers/CompositeLayerPlate.h
@@ -14,7 +14,7 @@ class UPrimitiveComponent;
 class UTexture;
 class FObjectPreSaveContext;
 class UTextureRenderTarget2D;
-class UCompositePassCenteredScale;
+class UCompositePassTransform2D;
 
 /** Plate sampling & rendering mode. */
 UENUM(BlueprintType)
@@ -329,7 +329,7 @@ private:
 
 	/** Internal pass to automatically remove overscan from the media texture. Not serialized. */
 	UPROPERTY(Transient, NonTransactional)
-	TObjectPtr<UCompositePassCenteredScale> OverscanPass;
+	TObjectPtr<UCompositePassTransform2D> OverscanPass;
 
 	/** Cached number of valid media passes. */
 	int32 CachedValidMediaPasses = 0;
diff --git a/Engine/Plugins/Compositing/Composite/Source/Composite/Public/Passes/CompositePassCenteredScale.h b/Engine/Plugins/Compositing/Composite/Source/Composite/Public/Passes/CompositePassCenteredScale.h
deleted file mode 100644
diff --git a/Engine/Plugins/Compositing/Composite/Source/Composite/Public/Passes/CompositePassTransform2D.h b/Engine/Plugins/Compositing/Composite/Source/Composite/Public/Passes/CompositePassTransform2D.h
new file mode 100644
--- /dev/null
+++ b/Engine/Plugins/Compositing/Composite/Source/Composite/Public/Passes/CompositePassTransform2D.h
@@ -0,0 +1,139 @@
+// Copyright Epic Games, Inc. All Rights Reserved.
+
+#pragma once
+
+#include "CompositePassBase.h"
+
+#include "CompositePassTransform2D.generated.h"
+
+#define UE_API COMPOSITE_API
+
+/** Primary scale calculation mode. */
+UENUM(BlueprintType)
+enum class ECompositePassScaleMode : uint8
+{
+	/** Disabled */
+	None,
+
+	/** Automatically derive aspect ratio from the parent layer media texture & composite actor camera. */
+	Automatic,
+
+	/** Scale based on the difference between container & content aspect ratios. */
+	AspectRatio,
+
+	/** Remove black bar padding by specifying pixels to crop on each axis. */
+	Unpad,
+
+	/** Manually define the scale factor.*/
+	Manual,
+};
+
+/** 2D transform pass: scale, rotate, and translate footage in UV space. */
+UCLASS(MinimalAPI, BlueprintType, Blueprintable, EditInlineNew, CollapseCategories, meta=(DisplayName="Transform 2D"))
+class UCompositePassTransform2D : public UCompositePassBase
+{
+	GENERATED_BODY()
+
+public:
+	/** Constructor */
+	UE_API UCompositePassTransform2D(const FObjectInitializer& ObjectInitializer);
+
+	/** Destructor */
+	UE_API ~UCompositePassTransform2D();
+
+	UE_API virtual FCompositeCorePassProxy* GetProxy(const UE::CompositeCore::FPassInputDecl& InputDecl, FCompositeTraversalContext& InContext, FSceneRenderingBulkObjectAllocator& InFrameAllocator) const override;
+
+	UE_API virtual bool GetIsActive() const override;
+
+	/** Calculate the final texture UV scale. */
+	UFUNCTION(BlueprintCallable, Category = "Composite")
+	UE_API FVector2f CalculateScale() const;
+
+public:
+	/** Normalized pivot point for rotation and scale. (0.5, 0.5) = viewport center. */
+	UPROPERTY(EditAnywhere, BlueprintReadWrite, Interp, Category = "Composite")
+	FVector2f Pivot;
+
+	/** How to scale-fit the video content into the camera viewport. */
+	UPROPERTY(EditAnywhere, BlueprintReadWrite, Interp, Category = "Composite")
+	ECompositePassScaleMode ScaleMode;
+
+	/** Container (source) aspect ratio, as width:height (e.g. 16:9). Automatic derives this from the plate texture. */
+	UPROPERTY(EditAnywhere, BlueprintReadWrite, Interp, Category = "Composite", meta = (EditCondition = "ScaleMode == ECompositePassScaleMode::AspectRatio", EditConditionHides))
+	FVector2f ContainerAspectRatio;
+
+	/** Content (target) aspect ratio, typically matching the cine camera, as width:height (e.g. 16:9). Automatic derives this from the camera component. */
+	UPROPERTY(EditAnywhere, BlueprintReadWrite, Interp, Category = "Composite", meta = (EditCondition = "ScaleMode == ECompositePassScaleMode::AspectRatio", EditConditionHides))
+	FVector2f ContentAspectRatio;
+
+	/** How to resolve aspect ratio mismatch. Scales only the mismatched axis, or both axes equally. */
+	UPROPERTY(EditAnywhere, BlueprintReadWrite, Interp, Category = "Composite", meta = (EditCondition = "ScaleMode == ECompositePassScaleMode::Automatic || ScaleMode == ECompositePassScaleMode::AspectRatio", EditConditionHides))
+	bool bScaleSingleAxis;
+
+	/** Pixels of black bar padding to crop from each side, per axis. For example, (656, 0) removes 656px of pillarbox from left and right. */
+	UPROPERTY(EditAnywhere, BlueprintReadWrite, Interp, Category = "Composite", meta = (EditCondition = "ScaleMode == ECompositePassScaleMode::Unpad", EditConditionHides, ClampMin = "0", UIMax = "500"))
+	FIntPoint UnpadPixels;
+
+	/** Manual scale factor per axis. Values > 1 zoom in, values < 1 zoom out. */
+	UPROPERTY(EditAnywhere, BlueprintReadWrite, Interp, Category = "Composite", meta = (EditCondition = "ScaleMode == ECompositePassScaleMode::Manual", EditConditionHides))
+	FVector2f ManualScale;
+
+	/** Automatically remove overscan by reading the camera's overscan value. */
+	UPROPERTY(BlueprintReadWrite, Category = "Composite")
+	bool bRemoveOverscan;
+
+	/** Counter-clockwise rotation in degrees. */
+	UPROPERTY(EditAnywhere, BlueprintReadWrite, Interp, Category = "Composite", meta = (UIMin="-360.0", UIMax="360.0", ClampMin="-360.0", ClampMax="360.0"))
+	float RotationAngle;
+
+	/** 2D translation offset in normalized viewport coordinates. Y > 0 moves the image up. */
+	UPROPERTY(EditAnywhere, BlueprintReadWrite, Interp, Category = "Composite")
+	FVector2f Translation;
+
+private:
+	/** Return the plate texture size from the owning plate layer, or ZeroValue if unavailable. */
+	FIntVector2 GetPlateTextureSize() const;
+
+	/** Cached scale UV computed in GetIsActive() and reused in GetProxy(). Reset after consumption. */
+	mutable TOptional<FVector2f> CachedScaleUV;
+};
+
+namespace UE
+{
+	namespace Composite
+	{
+		namespace Private
+		{
+			using namespace CompositeCore;
+
+			/** Render-thread proxy for the 2D transform pass. Applies scale, rotation, and translation around a configurable pivot in UV space. */
+			class FTransform2DPassProxy : public FCompositeCorePassProxy
+			{
+			public:
+				IMPLEMENT_COMPOSITE_PASS(FTransform2DPassProxy);
+
+				using FCompositeCorePassProxy::FCompositeCorePassProxy;
+
+				/** Adds the 2D transform RDG pass; bilinear-resamples the input with the full affine transform and returns the transformed output. */
+				FPassTexture Add(FRDGBuilder& GraphBuilder, const FSceneView& InView, const FPassInputArray& Inputs, const FPassContext& PassContext) const override;
+
+				/** Normalized pivot point for rotation and scale. */
+				FVector2f Pivot = FVector2f(0.5f, 0.5f);
+
+				/** Per-axis UV scale factor. (1, 1) is identity. */
+				FVector2f ScaleUV = FVector2f::One();
+
+				/** Rotation in radians. */
+				float RotationRadians = 0.0f;
+
+				/** 2D translation in normalized viewport coordinates. */
+				FVector2f Translation = FVector2f::ZeroVector;
+
+				/** Viewport aspect ratio (width/height) for rotation correction. */
+				float AspectRatio = 1.0f;
+			};
+		}
+	}
+}
+
+#undef UE_API
